Platform; cleanup OpenSLES audio on audio device teardown to avoid resource leak.

diff --git a/Engine/Source/Runtime/Android/AndroidAudio/Private/AndroidAudioDevice.cpp b/Engine/Source/Runtime/Android/AndroidAudio/Private/AndroidAudioDevice.cpp
--- a/Engine/Source/Runtime/Android/AndroidAudio/Private/AndroidAudioDevice.cpp
+++ b/Engine/Source/Runtime/Android/AndroidAudio/Private/AndroidAudioDevice.cpp
@@ -39,8 +39,22 @@ void FSLESAudioDevice::Teardown( void )
 	{
 		delete Sources[ i ];
 	}
+
+	UE_LOG( LogAndroidAudio, Warning, TEXT("OpenSLES Tearing Down HW"));
 	
-	//@todo android: teardown OpenSLES
+	// Teardown OpenSLES..
+	// Destroy the SLES objects in reverse order of creation:
+	if (SL_OutputMixObject)
+	{
+		(*SL_OutputMixObject)->Destroy(SL_OutputMixObject);
+		SL_OutputMixObject = NULL;
+	}
+	if (SL_EngineObject)
+	{
+		(*SL_EngineObject)->Destroy(SL_EngineObject);
+		SL_EngineObject = NULL;
+		SL_EngineEngine = NULL;
+	}
 }
 
 /*------------------------------------------------------------------------------------
