Ensure that precaching of a sound completes instead of failing to start for all platforms (not just XAudio done in CL# 2082171)
TTP #335967

diff --git a/Engine/Source/Runtime/Android/AndroidAudio/Private/AndroidAudioBuffer.cpp b/Engine/Source/Runtime/Android/AndroidAudio/Private/AndroidAudioBuffer.cpp
--- a/Engine/Source/Runtime/Android/AndroidAudio/Private/AndroidAudioBuffer.cpp
+++ b/Engine/Source/Runtime/Android/AndroidAudio/Private/AndroidAudioBuffer.cpp
@@ -111,12 +111,7 @@ FSLESSoundBuffer* FSLESSoundBuffer::CreateNativeBuffer( FSLESAudioDevice* AudioD
 	// Check to see if thread has finished decompressing on the other thread
 	if( InWave->AudioDecompressor != NULL )
 	{
-		if( !InWave->AudioDecompressor->IsDone() )
-		{
-			// Don't play this sound just yet
-			UE_LOG(LogAndroidAudio, Log, TEXT( "Waiting for sound to decompress: %s mode(%d)" ), *InWave->GetName(), int32(InWave->DecompressionType) );
-			return( NULL );
-		}
+		InWave->AudioDecompressor->EnsureCompletion();
 
 		// Remove the decompressor
 		delete InWave->AudioDecompressor;
diff --git a/Engine/Source/Runtime/Linux/ALAudio/Private/ALAudioBuffer.cpp b/Engine/Source/Runtime/Linux/ALAudio/Private/ALAudioBuffer.cpp
--- a/Engine/Source/Runtime/Linux/ALAudio/Private/ALAudioBuffer.cpp
+++ b/Engine/Source/Runtime/Linux/ALAudio/Private/ALAudioBuffer.cpp
@@ -94,12 +94,7 @@ FALSoundBuffer* FALSoundBuffer::CreateNativeBuffer( FALAudioDevice* AudioDevice,
 	// Check to see if thread has finished decompressing on the other thread
 	if (Wave->AudioDecompressor != NULL)
 	{
-		if (!Wave->AudioDecompressor->IsDone())
-		{
-			// Don't play this sound just yet
-			UE_LOG(LogAudio, Log, TEXT("Waiting for sound to decompress: %s"), *Wave->GetName());
-			return(NULL);
-		}
+		Wave->AudioDecompressor->EnsureCompletion();
 
 		// Remove the decompressor
 		delete Wave->AudioDecompressor;
diff --git a/Engine/Source/Runtime/Mac/CoreAudio/Private/CoreAudioBuffer.cpp b/Engine/Source/Runtime/Mac/CoreAudio/Private/CoreAudioBuffer.cpp
--- a/Engine/Source/Runtime/Mac/CoreAudio/Private/CoreAudioBuffer.cpp
+++ b/Engine/Source/Runtime/Mac/CoreAudio/Private/CoreAudioBuffer.cpp
@@ -246,12 +246,7 @@ FCoreAudioSoundBuffer* FCoreAudioSoundBuffer::CreateNativeBuffer( FCoreAudioDevi
 	// Check to see if thread has finished decompressing on the other thread
 	if( Wave->AudioDecompressor != NULL )
 	{
-		if (!Wave->AudioDecompressor->IsDone())
-		{
-			// Don't play this sound just yet
-			UE_LOG(LogAudio, Log, TEXT( "Waiting for sound to decompress: %s" ), *Wave->GetName() );
-			return( NULL );
-		}
+		Wave->AudioDecompressor->EnsureCompletion();
 		
 		// Remove the decompressor
 		delete Wave->AudioDecompressor;
