Fix audio playback cutting in and out, especially for ambient sounds. Volume was being cut artificially cut to zero below a certain threshold. This removes that cut to bring it on-par with the other platforms. And also syncs the audio source update to mirror other platforms more closely (for example adds missing global volume attenuation).

#ttp 334725

diff --git a/Engine/Source/Runtime/Android/AndroidAudio/Private/AndroidAudioSource.cpp b/Engine/Source/Runtime/Android/AndroidAudio/Private/AndroidAudioSource.cpp
--- a/Engine/Source/Runtime/Android/AndroidAudio/Private/AndroidAudioSource.cpp
+++ b/Engine/Source/Runtime/Android/AndroidAudio/Private/AndroidAudioSource.cpp
@@ -312,17 +312,24 @@ void FSLESSoundSource::Update( void )
 	{
 		return;
 	}
-	
+
 	float Volume = WaveInstance->Volume * WaveInstance->VolumeMultiplier;
 	if( SetStereoBleed() )
 	{
 		// Emulate the bleed to rear speakers followed by stereo fold down
 		Volume *= 1.25f;
 	}
-	
-	Volume		= FMath::Clamp<float>( Volume, 0.0f, 1.0f );
-	float Pitch = FMath::Clamp<float>( WaveInstance->Pitch, 0.4f, 2.0f );
-	
+	Volume *= GVolumeMultiplier;
+	Volume = FMath::Clamp(Volume, 0.0f, MAX_VOLUME);
+
+	const float Pitch = FMath::Clamp<float>(WaveInstance->Pitch, MIN_PITCH, MAX_PITCH);
+
+	// Set whether to apply reverb
+	SetReverbApplied(true);
+
+	// Set the HighFrequencyGain value
+	SetHighFrequencyGain();
+
 	FVector Location;
 	FVector	Velocity;
 	
@@ -344,17 +351,13 @@ void FSLESSoundSource::Update( void )
 	// Set volume & Pitch
 	// also Location & Velocity
 	
+	// Convert volume to millibels.
 	SLmillibel MaxMillibel = 0;
 	SLmillibel MinMillibel = -3000;
 	(*SL_VolumeInterface)->GetMaxVolumeLevel( SL_VolumeInterface, &MaxMillibel );
+	SLmillibel VolumeMillibel = (Volume * (MaxMillibel - MinMillibel)) + MinMillibel;
 	
-	// drop it down to an inaudible range
-	if( Volume < 0.1f )
-	{
-		MinMillibel = -10000;
-	}
-	
-	SLresult result = (*SL_VolumeInterface)->SetVolumeLevel(SL_VolumeInterface, ( Volume * ( MaxMillibel - MinMillibel ) ) + MinMillibel ); 
+	SLresult result = (*SL_VolumeInterface)->SetVolumeLevel(SL_VolumeInterface, VolumeMillibel);
 	check(SL_RESULT_SUCCESS == result);
 
 }
