Fixing issues with SoundWave memory and duplicated variables in some platforms.

Moved most of the contents of FAudioDevice::FreeResource to USoundWave::FreeResources so that we can be sure that all memory is cleared up during teardown.
Removed FreeResource function from some subclasses of audio device, as they weren't doing anything different than the base one.
Also removed ListSounds function from these same classes as that seems to have fallen behind current base class version and isn't used anyway.
Removed platform specific audio buffer storage so that there isn't confusion about which is being used and ensured that casts are used when necessary.
Fixed compile issue with STAT_AudioResourceCreationTime for HTML5.

diff --git a/Engine/Source/Runtime/Android/AndroidAudio/Private/AndroidAudioBuffer.cpp b/Engine/Source/Runtime/Android/AndroidAudio/Private/AndroidAudioBuffer.cpp
--- a/Engine/Source/Runtime/Android/AndroidAudio/Private/AndroidAudioBuffer.cpp
+++ b/Engine/Source/Runtime/Android/AndroidAudio/Private/AndroidAudioBuffer.cpp
@@ -31,11 +31,6 @@ FSLESSoundBuffer::FSLESSoundBuffer( FSLESAudioDevice* InAudioDevice ) :
  */
 FSLESSoundBuffer::~FSLESSoundBuffer( void )
 {
-	if( ResourceID )
-	{
-		AudioDevice->WaveBufferMap.Remove( ResourceID );
-	}
-	
 	FMemory::Free( AudioData);
 
 	if( DecompressionState )
@@ -125,16 +120,8 @@ FSLESSoundBuffer* FSLESSoundBuffer::CreateNativeBuffer( FSLESAudioDevice* AudioD
 	// Create new buffer.
 	Buffer = new FSLESSoundBuffer( AudioDevice );
 		
-	// Allocate new resource ID and assign to USoundNodeWave. A value of 0 (default) means not yet registered.
-	int32 ResourceID = AudioDevice->NextResourceID++;
-	Buffer->ResourceID = ResourceID;
-	InWave->ResourceID = ResourceID;
-		
-	AudioDevice->Buffers.Add( Buffer );
-	AudioDevice->WaveBufferMap.Add( ResourceID, Buffer );
-		
-	// Keep track of associated resource name.
-	Buffer->ResourceName	= InWave->GetPathName();		
+	AudioDevice->TrackResource(InWave, Buffer);
+
 	Buffer->NumChannels		= InWave->NumChannels;
 	Buffer->SampleRate		= InWave->SampleRate;
 
@@ -193,7 +180,7 @@ FSLESSoundBuffer* FSLESSoundBuffer::Init(  FSLESAudioDevice* AudioDevice ,USound
 		// Upload entire wav
 		if( InWave->ResourceID )
 		{
-			Buffer = AudioDevice->WaveBufferMap.FindRef( InWave->ResourceID );
+			Buffer = static_cast<FSLESSoundBuffer*>(AudioDevice->WaveBufferMap.FindRef( InWave->ResourceID ));
 		}
 
 		if( Buffer == NULL )
diff --git a/Engine/Source/Runtime/Android/AndroidAudio/Private/AndroidAudioDevice.cpp b/Engine/Source/Runtime/Android/AndroidAudio/Private/AndroidAudioDevice.cpp
--- a/Engine/Source/Runtime/Android/AndroidAudio/Private/AndroidAudioDevice.cpp
+++ b/Engine/Source/Runtime/Android/AndroidAudio/Private/AndroidAudioDevice.cpp
@@ -133,9 +133,6 @@ bool FSLESAudioDevice::InitializeHardware( void )
 	// Set up a default (nop) effects manager 
 	Effects = new FAudioEffectsManager( this );
 	
-	// Initialized.
-	NextResourceID = 1;
-
 	return true;
 }
 
@@ -152,79 +149,6 @@ void FSLESAudioDevice::Update( bool Realtime )
 	//@todo android: UDPATE LISTENERS - Android OpenSLES doesn't support 3D, is there anything to do here?
 }
 
-
-/**
- * Frees the bulk resource data associated with this SoundNodeWave.
- *
- * @param	SoundNodeWave	wave object to free associated bulk data
- */
-void FSLESAudioDevice::FreeResource( USoundWave* SoundNodeWave )
-{
-	// Just in case the data was created but never uploaded
-	if( SoundNodeWave->RawPCMData )
-	{
-		FMemory::Free( SoundNodeWave->RawPCMData );
-		SoundNodeWave->RawPCMData = NULL;
-	}
-	
-	// Find buffer for resident wavs
-	if( SoundNodeWave->ResourceID )
-	{
-		// Find buffer associated with resource id.
-		FSLESSoundBuffer* Buffer = WaveBufferMap.FindRef( SoundNodeWave->ResourceID );
-		if( Buffer )
-		{
-			// Remove from buffers array.
-			Buffers.Remove( Buffer );
-			
-			// See if it is being used by a sound source...
-			for( int32 SrcIndex = 0; SrcIndex < Sources.Num(); SrcIndex++ )
-			{
-				FSLESSoundSource* Src = ( FSLESSoundSource* )( Sources[ SrcIndex ] );
-				if( Src && Src->Buffer && ( Src->Buffer == Buffer ) )
-				{
-					Src->Stop();
-					break;
-				}
-			}
-			
-			delete Buffer;
-		}
-		
-		SoundNodeWave->ResourceID = 0;
-	}
-	
-	// .. or reference to compressed data
-	SoundNodeWave->RemoveAudioResource();
-	
-	// Stat housekeeping
-	//DEC_DWORD_STAT_BY( STAT_AudioMemorySize, SoundNodeWave->GetResourceSize() );
-	//DEC_DWORD_STAT_BY( STAT_AudioMemory, SoundNodeWave->GetResourceSize() );
-}
-
-
-
-/** 
- * Displays debug information about the loaded sounds
- */
-void FSLESAudioDevice::ListSounds( const TCHAR* Cmd, FOutputDevice& Ar )
-{
-	int	TotalSoundSize = 0;
-
-	Ar.Logf( TEXT( "Sound resources:" ) );
-
-	TArray<FSLESSoundBuffer*> AllSounds = Buffers;
-
-	for( int i = 0; i < AllSounds.Num(); ++i )
-	{
-		FSLESSoundBuffer* Buffer = AllSounds[i];
-		Ar.Logf( TEXT( "RawData: %8.2f Kb (%d channels at %d Hz) in sound %s" ), Buffer->GetSize() / 1024.0f, Buffer->GetNumChannels(), Buffer->SampleRate, *Buffer->ResourceName );
-		TotalSoundSize += Buffer->GetSize();
-	}
-
-	Ar.Logf( TEXT( "%8.2f Kb for %d sounds" ), TotalSoundSize / 1024.0f, AllSounds.Num() );
-}
-
 FSoundSource* FSLESAudioDevice::CreateSoundSource()
 {
 	return new FSLESSoundSource(this);
diff --git a/Engine/Source/Runtime/Android/AndroidAudio/Public/AndroidAudioDevice.h b/Engine/Source/Runtime/Android/AndroidAudio/Public/AndroidAudioDevice.h
--- a/Engine/Source/Runtime/Android/AndroidAudio/Public/AndroidAudioDevice.h
+++ b/Engine/Source/Runtime/Android/AndroidAudio/Public/AndroidAudioDevice.h
@@ -74,7 +74,7 @@ public:
 	 *
 	 * @param	InWave				USoundWave to use as template and wave source
 	 * @param	AudioDevice			Audio device to attach created buffer to
-	 * @return	FALSoundBuffer pointer if buffer creation succeeded, NULL otherwise
+	 * @return	FSLESSoundBuffer pointer if buffer creation succeeded, NULL otherwise
 	 */
 	static FSLESSoundBuffer* Init(  FSLESAudioDevice* ,USoundWave* );
 
@@ -261,19 +261,6 @@ public:
 	 */
 	virtual void Update( bool bGameTicking );
 
-	/** 
-	 * Lists all the loaded sounds and their memory footprint
-	 */
-	virtual void ListSounds( const TCHAR* Cmd, FOutputDevice& Ar );
-
-	/**
-	 * Frees the bulk resource data associated with this SoundNodeWave.
-	 *
-	 * @param	SoundNodeWave	wave object to free associated bulk data
-	 */
-	virtual void FreeResource( USoundWave* SoundWave );
-
-
 protected:
 
 	virtual FSoundSource* CreateSoundSource() override; 
@@ -289,12 +276,6 @@ protected:
 
 	/** The name of the OpenSL Device to open - defaults to "Generic Software" */
 	FString										DeviceName;
-	/** All loaded resident buffers */
-	TArray<FSLESSoundBuffer*>					Buffers;
-	/** Map from resource ID to sound buffer */
-	TMap<int, FSLESSoundBuffer*>				WaveBufferMap;
-	/** Next resource ID value used for registering USoundWave objects */
-	int											NextResourceID;
 
 	// SL specific
 
diff --git a/Engine/Source/Runtime/Engine/Private/AudioDevice.cpp b/Engine/Source/Runtime/Engine/Private/AudioDevice.cpp
--- a/Engine/Source/Runtime/Engine/Private/AudioDevice.cpp
+++ b/Engine/Source/Runtime/Engine/Private/AudioDevice.cpp
@@ -2377,18 +2377,6 @@ void FAudioDevice::FreeResource(USoundWave* SoundWave)
 
 		SoundWave->ResourceID = 0;
 	}
-
-	// Just in case the data was created but never uploaded
-	FMemory::Free(SoundWave->RawPCMData);
-	SoundWave->RawPCMData = NULL;
-
-	// Remove the compressed copy of the data
-	SoundWave->RemoveAudioResource();
-
-	// Stat housekeeping
-	DEC_DWORD_STAT_BY(STAT_AudioMemorySize, SoundWave->TrackedMemoryUsage);
-	DEC_DWORD_STAT_BY(STAT_AudioMemory, SoundWave->TrackedMemoryUsage);
-	SoundWave->TrackedMemoryUsage = 0;
 }
 
 void FAudioDevice::FreeBufferResource(FSoundBuffer* Buffer)
diff --git a/Engine/Source/Runtime/Engine/Private/SoundWave.cpp b/Engine/Source/Runtime/Engine/Private/SoundWave.cpp
--- a/Engine/Source/Runtime/Engine/Private/SoundWave.cpp
+++ b/Engine/Source/Runtime/Engine/Private/SoundWave.cpp
@@ -423,6 +423,18 @@ void USoundWave::FreeResources()
 		}
 	}
 
+	// Just in case the data was created but never uploaded
+	FMemory::Free(RawPCMData);
+	RawPCMData = NULL;
+
+	// Remove the compressed copy of the data
+	RemoveAudioResource();
+
+	// Stat housekeeping
+	DEC_DWORD_STAT_BY(STAT_AudioMemorySize, TrackedMemoryUsage);
+	DEC_DWORD_STAT_BY(STAT_AudioMemory, TrackedMemoryUsage);
+	TrackedMemoryUsage = 0;
+
 	SampleRate = 0;
 	Duration = 0.0f;
 	ResourceID = 0;
diff --git a/Engine/Source/Runtime/HTML5/HTML5Audio/Private/HTML5AudioBuffer.cpp b/Engine/Source/Runtime/HTML5/HTML5Audio/Private/HTML5AudioBuffer.cpp
--- a/Engine/Source/Runtime/HTML5/HTML5Audio/Private/HTML5AudioBuffer.cpp
+++ b/Engine/Source/Runtime/HTML5/HTML5Audio/Private/HTML5AudioBuffer.cpp
@@ -5,6 +5,7 @@
 #include "Engine.h"
 #include "IAudioFormat.h"
 #include "AudioDecompress.h"
+#include "ContentStreaming.h"
 
 /*------------------------------------------------------------------------------------
 	FALSoundBuffer.
@@ -26,11 +27,6 @@ FALSoundBuffer::FALSoundBuffer( FALAudioDevice* InAudioDevice )
  */
 FALSoundBuffer::~FALSoundBuffer( void )
 {
-	if( ResourceID )
-	{
-		AudioDevice->WaveBufferMap.Remove( ResourceID );
-	}
-
 	// Delete AL buffers.
 	alDeleteBuffers( 1, BufferIds );
 }
@@ -118,7 +114,7 @@ FALSoundBuffer* FALSoundBuffer::CreateNativeBuffer( FALAudioDevice* AudioDevice,
 	// Find the existing buffer if any
 	if( Wave->ResourceID )
 	{
-		Buffer = AudioDevice->WaveBufferMap.FindRef( Wave->ResourceID );
+		Buffer = static_cast<FALSoundBuffer*>(AudioDevice->WaveBufferMap.FindRef( Wave->ResourceID ));
 	}
 
 	if( Buffer == NULL )
@@ -130,16 +126,7 @@ FALSoundBuffer* FALSoundBuffer::CreateNativeBuffer( FALAudioDevice* AudioDevice,
 
 		AudioDevice->alError( TEXT( "RegisterSound" ) );
 
-		// Allocate new resource ID and assign to USoundNodeWave. A value of 0 (default) means not yet registered.
-		int ResourceID = AudioDevice->NextResourceID++;
-		Buffer->ResourceID = ResourceID;
-		Wave->ResourceID = ResourceID;
-
-		AudioDevice->Buffers.Add( Buffer );
-		AudioDevice->WaveBufferMap.FindOrAdd(ResourceID) =  Buffer ;
-
-		// Keep track of associated resource name.
-		Buffer->ResourceName = Wave->GetPathName();
+		AudioDevice->TrackResource(Wave, Buffer);
 
 		Buffer->InternalFormat = AudioDevice->GetInternalFormat( Wave->NumChannels );
 		Buffer->NumChannels = Wave->NumChannels;
diff --git a/Engine/Source/Runtime/HTML5/HTML5Audio/Private/HTML5AudioDevice.cpp b/Engine/Source/Runtime/HTML5/HTML5Audio/Private/HTML5AudioDevice.cpp
--- a/Engine/Source/Runtime/HTML5/HTML5Audio/Private/HTML5AudioDevice.cpp
+++ b/Engine/Source/Runtime/HTML5/HTML5Audio/Private/HTML5AudioDevice.cpp
@@ -194,9 +194,6 @@ bool FALAudioDevice::InitializeHardware( void )
 	// Set up a default (nop) effects manager 
 	Effects = new FAudioEffectsManager( this );
 
-	// Initialized.
-	NextResourceID = 1;
-
 	return true ;
 }
 
@@ -246,81 +243,6 @@ void FALAudioDevice::Update( bool Realtime )
 	alError( TEXT( "UALAudioDevice::Update" ) );
 }	
 
-
-/**
- * Frees the bulk resource data associated with this SoundNodeWave.
- *
- * @param	SoundNodeWave	wave object to free associated bulk data
- */
-void FALAudioDevice::FreeResource( USoundWave* SoundNodeWave )
-{
-	// Just in case the data was created but never uploaded
-	if( SoundNodeWave->RawPCMData )
-	{
-		FMemory::Free( SoundNodeWave->RawPCMData );
-		SoundNodeWave->RawPCMData = NULL;
-	}
-
-	// Find buffer for resident wavs
-	if( SoundNodeWave->ResourceID )
-	{
-		// Find buffer associated with resource id.
-		FALSoundBuffer* Buffer = WaveBufferMap.FindRef( SoundNodeWave->ResourceID );
-		if( Buffer )
-		{
-			// Remove from buffers array.
-			Buffers.Remove( Buffer );
-
-			// See if it is being used by a sound source...
-			for( int SrcIndex = 0; SrcIndex < Sources.Num(); SrcIndex++ )
-			{
-				FALSoundSource* Src = ( FALSoundSource* )( Sources[ SrcIndex ] );
-				if( Src && Src->Buffer && ( Src->Buffer == Buffer ) )
-				{
-
-					Src->Stop();
-					break;
-				}
-			}
-
-			delete Buffer;
-		}
-
-		SoundNodeWave->ResourceID = 0;
-	}
-
-	// .. or reference to compressed data
-	SoundNodeWave->RemoveAudioResource();
-
-	// Stat housekeeping
-	// @to-do 
-// 	DEC_DWORD_STAT_BY( STAT_AudioMemorySize, SoundNodeWave->GetResourceSize() );
-// 	DEC_DWORD_STAT_BY( STAT_AudioMemory, SoundNodeWave->GetResourceSize() );
-}
-
-
-
-/** 
- * Displays debug information about the loaded sounds
- */
-void FALAudioDevice::ListSounds( const TCHAR* Cmd, FOutputDevice& Ar )
-{
-	int	TotalSoundSize = 0;
-
-	Ar.Logf( TEXT( "Sound resources:" ) );
-
-	TArray<FALSoundBuffer*> AllSounds = Buffers;
-
-	for( int i = 0; i < AllSounds.Num(); ++i )
-	{
-		FALSoundBuffer* Buffer = AllSounds[i];
-		Ar.Logf( TEXT( "RawData: %8.2f Kb (%d channels at %d Hz) in sound %s" ), Buffer->GetSize() / 1024.0f, Buffer->GetNumChannels(), Buffer->SampleRate, *Buffer->ResourceName );
-		TotalSoundSize += Buffer->GetSize();
-	}
-
-	Ar.Logf( TEXT( "%8.2f Kb for %d sounds" ), TotalSoundSize / 1024.0f, AllSounds.Num() );
-}
-
 ALuint FALAudioDevice::GetInternalFormat( int NumChannels )
 {
 	ALuint InternalFormat = 0;
diff --git a/Engine/Source/Runtime/HTML5/HTML5Audio/Public/HTML5AudioDevice.h b/Engine/Source/Runtime/HTML5/HTML5Audio/Public/HTML5AudioDevice.h
--- a/Engine/Source/Runtime/HTML5/HTML5Audio/Public/HTML5AudioDevice.h
+++ b/Engine/Source/Runtime/HTML5/HTML5Audio/Public/HTML5AudioDevice.h
@@ -220,18 +220,6 @@ public:
 	 */
 	virtual void Update( bool bGameTicking );
 
-	/** 
-	 * Lists all the loaded sounds and their memory footprint
-	 */
-	virtual void ListSounds( const TCHAR* Cmd, FOutputDevice& Ar );
-
-	/**
-	 * Frees the bulk resource data associated with this SoundNodeWave.
-	 *
-	 * @param	SoundNodeWave	wave object to free associated bulk data
-	 */
-	virtual void FreeResource( USoundWave* SoundWave );
-
 	virtual bool HasCompressedAudioInfoClass(USoundWave* SoundWave) override;
 
 	virtual class ICompressedAudioInfo* CreateCompressedAudioInfo(USoundWave* SoundWave) override;
@@ -254,13 +242,6 @@ protected:
 
 	// Variables.
 
-	/** All loaded resident buffers */
-	TArray<FALSoundBuffer*>						Buffers;
-	/** Map from resource ID to sound buffer */
-	TMap<int, FALSoundBuffer*>					WaveBufferMap;
-	/** Next resource ID value used for registering USoundWave objects */
-	int											NextResourceID;
-
 	// AL specific
 
 	/** Device/context used to play back sounds (static so it can be initialized early) */
diff --git a/Engine/Source/Runtime/Linux/ALAudio/Private/ALAudioBuffer.cpp b/Engine/Source/Runtime/Linux/ALAudio/Private/ALAudioBuffer.cpp
--- a/Engine/Source/Runtime/Linux/ALAudio/Private/ALAudioBuffer.cpp
+++ b/Engine/Source/Runtime/Linux/ALAudio/Private/ALAudioBuffer.cpp
@@ -26,11 +26,6 @@ FALSoundBuffer::FALSoundBuffer( FALAudioDevice* InAudioDevice )
  */
 FALSoundBuffer::~FALSoundBuffer( void )
 {
-	if( ResourceID )
-	{
-		AudioDevice->WaveBufferMap.Remove( ResourceID );
-	}
-
 	// Delete AL buffers.
 	alDeleteBuffers( 1, BufferIds );
 }
@@ -133,7 +128,7 @@ FALSoundBuffer* FALSoundBuffer::CreateNativeBuffer( FALAudioDevice* AudioDevice,
 	// Find the existing buffer if any
 	if (Wave->ResourceID)
 	{
-		Buffer = AudioDevice->WaveBufferMap.FindRef( Wave->ResourceID );
+		Buffer = static_cast<FALSoundBuffer*>(AudioDevice->WaveBufferMap.FindRef( Wave->ResourceID ));
 	}
 
 	if (Buffer == nullptr)
@@ -145,16 +140,7 @@ FALSoundBuffer* FALSoundBuffer::CreateNativeBuffer( FALAudioDevice* AudioDevice,
 
 		AudioDevice->alError(TEXT("RegisterSound"));
 
-		// Allocate new resource ID and assign to USoundNodeWave. A value of 0 (default) means not yet registered.
-		int ResourceID = AudioDevice->NextResourceID++;
-		Buffer->ResourceID = ResourceID;
-		Wave->ResourceID = ResourceID;
-
-		AudioDevice->Buffers.Add(Buffer);
-		AudioDevice->WaveBufferMap.FindOrAdd(ResourceID) =  Buffer;
-
-		// Keep track of associated resource name.
-		Buffer->ResourceName = Wave->GetPathName();
+		AudioDevice->TrackResource(Wave, Buffer);
 
 		Buffer->InternalFormat = AudioDevice->GetInternalFormat(Wave->NumChannels);
 		Buffer->NumChannels = Wave->NumChannels;
diff --git a/Engine/Source/Runtime/Linux/ALAudio/Private/ALAudioDevice.cpp b/Engine/Source/Runtime/Linux/ALAudio/Private/ALAudioDevice.cpp
--- a/Engine/Source/Runtime/Linux/ALAudio/Private/ALAudioDevice.cpp
+++ b/Engine/Source/Runtime/Linux/ALAudio/Private/ALAudioDevice.cpp
@@ -193,9 +193,6 @@ bool FALAudioDevice::InitializeHardware( void )
 	// Set up a default (nop) effects manager 
 	Effects = new FAudioEffectsManager( this );
 
-	// Initialized.
-	NextResourceID = 1;
-
 	return true ;
 }
 
@@ -245,81 +242,6 @@ void FALAudioDevice::Update( bool Realtime )
 	alError( TEXT( "UALAudioDevice::Update" ) );
 }	
 
-
-/**
- * Frees the bulk resource data associated with this SoundNodeWave.
- *
- * @param	SoundNodeWave	wave object to free associated bulk data
- */
-void FALAudioDevice::FreeResource( USoundWave* SoundNodeWave )
-{
-	// Just in case the data was created but never uploaded
-	if( SoundNodeWave->RawPCMData )
-	{
-		FMemory::Free( SoundNodeWave->RawPCMData );
-		SoundNodeWave->RawPCMData = NULL;
-	}
-
-	// Find buffer for resident wavs
-	if( SoundNodeWave->ResourceID )
-	{
-		// Find buffer associated with resource id.
-		FALSoundBuffer* Buffer = WaveBufferMap.FindRef( SoundNodeWave->ResourceID );
-		if( Buffer )
-		{
-			// Remove from buffers array.
-			Buffers.Remove( Buffer );
-
-			// See if it is being used by a sound source...
-			for( int SrcIndex = 0; SrcIndex < Sources.Num(); SrcIndex++ )
-			{
-				FALSoundSource* Src = ( FALSoundSource* )( Sources[ SrcIndex ] );
-				if( Src && Src->Buffer && ( Src->Buffer == Buffer ) )
-				{
-
-					Src->Stop();
-					break;
-				}
-			}
-
-			delete Buffer;
-		}
-
-		SoundNodeWave->ResourceID = 0;
-	}
-
-	// .. or reference to compressed data
-	SoundNodeWave->RemoveAudioResource();
-
-	// Stat housekeeping
-	// @to-do 
-// 	DEC_DWORD_STAT_BY( STAT_AudioMemorySize, SoundNodeWave->GetResourceSize() );
-// 	DEC_DWORD_STAT_BY( STAT_AudioMemory, SoundNodeWave->GetResourceSize() );
-}
-
-
-
-/** 
- * Displays debug information about the loaded sounds
- */
-void FALAudioDevice::ListSounds( const TCHAR* Cmd, FOutputDevice& Ar )
-{
-	int	TotalSoundSize = 0;
-
-	Ar.Logf( TEXT( "Sound resources:" ) );
-
-	TArray<FALSoundBuffer*> AllSounds = Buffers;
-
-	for( int i = 0; i < AllSounds.Num(); ++i )
-	{
-		FALSoundBuffer* Buffer = AllSounds[i];
-		Ar.Logf( TEXT( "RawData: %8.2f Kb (%d channels at %d Hz) in sound %s" ), Buffer->GetSize() / 1024.0f, Buffer->GetNumChannels(), Buffer->SampleRate, *Buffer->ResourceName );
-		TotalSoundSize += Buffer->GetSize();
-	}
-
-	Ar.Logf( TEXT( "%8.2f Kb for %d sounds" ), TotalSoundSize / 1024.0f, AllSounds.Num() );
-}
-
 ALuint FALAudioDevice::GetInternalFormat( int NumChannels )
 {
 	ALuint InternalFormat = 0;
diff --git a/Engine/Source/Runtime/Linux/ALAudio/Public/ALAudioDevice.h b/Engine/Source/Runtime/Linux/ALAudio/Public/ALAudioDevice.h
--- a/Engine/Source/Runtime/Linux/ALAudio/Public/ALAudioDevice.h
+++ b/Engine/Source/Runtime/Linux/ALAudio/Public/ALAudioDevice.h
@@ -218,18 +218,6 @@ public:
 	 */
 	virtual void Update( bool bGameTicking );
 
-	/** 
-	 * Lists all the loaded sounds and their memory footprint
-	 */
-	virtual void ListSounds( const TCHAR* Cmd, FOutputDevice& Ar );
-
-	/**
-	 * Frees the bulk resource data associated with this SoundNodeWave.
-	 *
-	 * @param	SoundNodeWave	wave object to free associated bulk data
-	 */
-	virtual void FreeResource( USoundWave* SoundWave );
-
 	virtual bool HasCompressedAudioInfoClass(USoundWave* SoundWave) override;
 
 	virtual class ICompressedAudioInfo* CreateCompressedAudioInfo(USoundWave* SoundWave) override;
@@ -252,13 +240,6 @@ protected:
 
 	// Variables.
 
-	/** All loaded resident buffers */
-	TArray<FALSoundBuffer*>						Buffers;
-	/** Map from resource ID to sound buffer */
-	TMap<int, FALSoundBuffer*>					WaveBufferMap;
-	/** Next resource ID value used for registering USoundWave objects */
-	int											NextResourceID;
-
 	// AL specific
 
 	/** Device/context used to play back sounds (static so it can be initialized early) */
diff --git a/Engine/Source/Runtime/Mac/CoreAudio/Private/CoreAudioBuffer.cpp b/Engine/Source/Runtime/Mac/CoreAudio/Private/CoreAudioBuffer.cpp
--- a/Engine/Source/Runtime/Mac/CoreAudio/Private/CoreAudioBuffer.cpp
+++ b/Engine/Source/Runtime/Mac/CoreAudio/Private/CoreAudioBuffer.cpp
@@ -137,7 +137,7 @@ void FCoreAudioSoundBuffer::Seek( const float SeekTime )
  *
  * @param InWave		USoundWave to use as template and wave source
  * @param AudioDevice	audio device to attach created buffer to
- * @return FALSoundBuffer pointer if buffer creation succeeded, NULL otherwise
+ * @return FCoreAudioSoundBuffer pointer if buffer creation succeeded, NULL otherwise
  */
 FCoreAudioSoundBuffer* FCoreAudioSoundBuffer::CreateQueuedBuffer( FCoreAudioDevice* CoreAudioDevice, USoundWave* Wave )
 {
@@ -181,7 +181,7 @@ FCoreAudioSoundBuffer* FCoreAudioSoundBuffer::CreateQueuedBuffer( FCoreAudioDevi
  *
  * @param InWave		USoundWave to use as template and wave source
  * @param AudioDevice	audio device to attach created buffer to
- * @return FALSoundBuffer pointer if buffer creation succeeded, NULL otherwise
+ * @return FCoreAudioSoundBuffer pointer if buffer creation succeeded, NULL otherwise
  */
 FCoreAudioSoundBuffer* FCoreAudioSoundBuffer::CreateProceduralBuffer( FCoreAudioDevice* CoreAudioDevice, USoundWave* Wave )
 {
@@ -206,7 +206,7 @@ FCoreAudioSoundBuffer* FCoreAudioSoundBuffer::CreateProceduralBuffer( FCoreAudio
  *
  * @param InWave		USoundWave to use as template and wave source
  * @param AudioDevice	audio device to attach created buffer to
- * @return FALSoundBuffer pointer if buffer creation succeeded, NULL otherwise
+ * @return FCoreAudioSoundBuffer pointer if buffer creation succeeded, NULL otherwise
  */
 FCoreAudioSoundBuffer* FCoreAudioSoundBuffer::CreatePreviewBuffer( FCoreAudioDevice* CoreAudioDevice, USoundWave* Wave, FCoreAudioSoundBuffer* Buffer )
 {
diff --git a/Engine/Source/Runtime/Windows/XAudio2/Private/XAudio2Buffer.cpp b/Engine/Source/Runtime/Windows/XAudio2/Private/XAudio2Buffer.cpp
--- a/Engine/Source/Runtime/Windows/XAudio2/Private/XAudio2Buffer.cpp
+++ b/Engine/Source/Runtime/Windows/XAudio2/Private/XAudio2Buffer.cpp
@@ -316,7 +316,7 @@ void FXAudio2SoundBuffer::Seek( const float SeekTime )
  *
  * @param InWave		USoundWave to use as template and wave source
  * @param AudioDevice	audio device to attach created buffer to
- * @return FALSoundBuffer pointer if buffer creation succeeded, NULL otherwise
+ * @return FXAudio2SoundBuffer pointer if buffer creation succeeded, NULL otherwise
  */
 FXAudio2SoundBuffer* FXAudio2SoundBuffer::CreateQueuedBuffer( FXAudio2Device* XAudio2Device, USoundWave* Wave )
 {
@@ -360,7 +360,7 @@ FXAudio2SoundBuffer* FXAudio2SoundBuffer::CreateQueuedBuffer( FXAudio2Device* XA
  *
  * @param InWave		USoundWave to use as template and wave source
  * @param AudioDevice	audio device to attach created buffer to
- * @return FALSoundBuffer pointer if buffer creation succeeded, NULL otherwise
+ * @return FXAudio2SoundBuffer pointer if buffer creation succeeded, NULL otherwise
  */
 FXAudio2SoundBuffer* FXAudio2SoundBuffer::CreateProceduralBuffer( FXAudio2Device* XAudio2Device, USoundWave* Wave )
 {
@@ -385,7 +385,7 @@ FXAudio2SoundBuffer* FXAudio2SoundBuffer::CreateProceduralBuffer( FXAudio2Device
  *
  * @param InWave		USoundWave to use as template and wave source
  * @param AudioDevice	audio device to attach created buffer to
- * @return FALSoundBuffer pointer if buffer creation succeeded, NULL otherwise
+ * @return FXAudio2SoundBuffer pointer if buffer creation succeeded, NULL otherwise
  */
 FXAudio2SoundBuffer* FXAudio2SoundBuffer::CreatePreviewBuffer( FXAudio2Device* XAudio2Device, USoundWave* Wave, FXAudio2SoundBuffer* Buffer )
 {
@@ -418,7 +418,7 @@ FXAudio2SoundBuffer* FXAudio2SoundBuffer::CreatePreviewBuffer( FXAudio2Device* X
  *
  * @param InWave		USoundWave to use as template and wave source
  * @param AudioDevice	audio device to attach created buffer to
- * @return FALSoundBuffer pointer if buffer creation succeeded, NULL otherwise
+ * @return FXAudio2SoundBuffer pointer if buffer creation succeeded, NULL otherwise
  */
 FXAudio2SoundBuffer* FXAudio2SoundBuffer::CreateNativeBuffer( FXAudio2Device* XAudio2Device, USoundWave* Wave )
 {
diff --git a/Engine/Source/Runtime/Windows/XAudio2/Private/XAudio2Support.h b/Engine/Source/Runtime/Windows/XAudio2/Private/XAudio2Support.h
--- a/Engine/Source/Runtime/Windows/XAudio2/Private/XAudio2Support.h
+++ b/Engine/Source/Runtime/Windows/XAudio2/Private/XAudio2Support.h
@@ -138,7 +138,7 @@ public:
 	 *
 	 * @param InWave		USoundWave to use as template and wave source
 	 * @param AudioDevice	audio device to attach created buffer to
-	 * @return FALSoundBuffer pointer if buffer creation succeeded, NULL otherwise
+	 * @return FXAudio2SoundBuffer pointer if buffer creation succeeded, NULL otherwise
 	 */
 	static FXAudio2SoundBuffer* CreateQueuedBuffer( FXAudio2Device* XAudio2Device, USoundWave* Wave );
 
@@ -147,7 +147,7 @@ public:
 	 *
 	 * @param InWave		USoundWave to use as template and wave source
 	 * @param AudioDevice	audio device to attach created buffer to
-	 * @return FALSoundBuffer pointer if buffer creation succeeded, NULL otherwise
+	 * @return FXAudio2SoundBuffer pointer if buffer creation succeeded, NULL otherwise
 	 */
 	static FXAudio2SoundBuffer* CreateProceduralBuffer( FXAudio2Device* XAudio2Device, USoundWave* Wave );
 
@@ -156,7 +156,7 @@ public:
 	 *
 	 * @param InWave		USoundWave to use as template and wave source
 	 * @param AudioDevice	audio device to attach created buffer to
-	 * @return FALSoundBuffer pointer if buffer creation succeeded, NULL otherwise
+	 * @return FXAudio2SoundBuffer pointer if buffer creation succeeded, NULL otherwise
 	 */
 	static FXAudio2SoundBuffer* CreatePreviewBuffer( FXAudio2Device* XAudio2Device, USoundWave* Wave, FXAudio2SoundBuffer* Buffer );
 
@@ -165,7 +165,7 @@ public:
 	 *
 	 * @param InWave		USoundWave to use as template and wave source
 	 * @param AudioDevice	audio device to attach created buffer to
-	 * @return FALSoundBuffer pointer if buffer creation succeeded, NULL otherwise
+	 * @return FXAudio2SoundBuffer pointer if buffer creation succeeded, NULL otherwise
 	 */
 	static FXAudio2SoundBuffer* CreateNativeBuffer( FXAudio2Device* XAudio2Device, USoundWave* Wave );
 
