Removed unneeded ResourceID from FLESSoundbuffer, fixes freed resource issue.
#UE4
#Android

diff --git a/Engine/Source/Runtime/Android/AndroidAudio/Public/AndroidAudioDevice.h b/Engine/Source/Runtime/Android/AndroidAudio/Public/AndroidAudioDevice.h
--- a/Engine/Source/Runtime/Android/AndroidAudio/Public/AndroidAudioDevice.h
+++ b/Engine/Source/Runtime/Android/AndroidAudio/Public/AndroidAudioDevice.h
@@ -110,8 +110,6 @@ public:
 	FSLESAudioDevice*			AudioDevice;
 	/** Data */
 	uint8*					AudioData;
-	/** Resource ID of associated USoundeWave */
-	int						ResourceID;
 	/** Human readable name of resource, most likely name of UObject associated during caching. */
 	FString					ResourceName;
 
diff --git a/Engine/Source/Runtime/Engine/Private/SoundWave.cpp b/Engine/Source/Runtime/Engine/Private/SoundWave.cpp
--- a/Engine/Source/Runtime/Engine/Private/SoundWave.cpp
+++ b/Engine/Source/Runtime/Engine/Private/SoundWave.cpp
@@ -424,8 +424,11 @@ void USoundWave::FreeResources()
 	}
 
 	// Just in case the data was created but never uploaded
-	FMemory::Free(RawPCMData);
-	RawPCMData = NULL;
+	if (RawPCMData != NULL)
+	{
+		FMemory::Free(RawPCMData);
+		RawPCMData = NULL;
+	}
 
 	// Remove the compressed copy of the data
 	RemoveAudioResource();
