#ttp 338505 MobileMVP: Audio: Procedural Sound not supported on Android

Updated android audio device/buffers/souce to support creation of procedural data from a sound source. Tested with MOD replay in QA test level.

diff --git a/Engine/Source/Runtime/Android/AndroidAudio/Private/AndroidAudioBuffer.cpp b/Engine/Source/Runtime/Android/AndroidAudio/Private/AndroidAudioBuffer.cpp
--- a/Engine/Source/Runtime/Android/AndroidAudio/Private/AndroidAudioBuffer.cpp
+++ b/Engine/Source/Runtime/Android/AndroidAudio/Private/AndroidAudioBuffer.cpp
@@ -136,6 +136,33 @@ FSLESSoundBuffer* FSLESSoundBuffer::CreateNativeBuffer( FSLESAudioDevice* AudioD
 	return Buffer;
 }
 
+/**
+* Static function used to create an Audio buffer and dynamically upload procedural data to.
+*
+* @param InWave		USoundWave to use as template and wave source
+* @param AudioDevice	audio device to attach created buffer to
+* @return FSLESSoundBuffer pointer if buffer creation succeeded, NULL otherwise
+*/
+FSLESSoundBuffer* FSLESSoundBuffer::CreateProceduralBuffer(FSLESAudioDevice* AudioDevice, USoundWave* InWave)
+{
+	FSLESSoundBuffer* Buffer = new FSLESSoundBuffer(AudioDevice);
+
+	// Setup any default information
+	Buffer->DecompressionState = NULL;
+	Buffer->AudioData = NULL;
+	Buffer->BufferSize = 0;
+	Buffer->Format = SoundFormat_PCMRT;
+	Buffer->NumChannels = InWave->NumChannels;
+	Buffer->SampleRate = InWave->SampleRate;
+	
+	InWave->RawPCMData = NULL;
+
+	// No tracking of this resource as it's temporary
+	Buffer->ResourceID = 0;
+	InWave->ResourceID = 0;
+
+	return Buffer;
+}
 
 /**
  * Static function used to create a buffer.
@@ -191,10 +218,14 @@ FSLESSoundBuffer* FSLESSoundBuffer::Init(  FSLESAudioDevice* AudioDevice ,USound
 		// Always create a new buffer for streaming ogg vorbis data
 		Buffer = CreateQueuedBuffer( AudioDevice, InWave );
 		break;
-		
+	
+	case DTYPE_Procedural:
+		// New buffer for procedural data
+		Buffer = CreateProceduralBuffer(AudioDevice, InWave);
+		break;
+
 	case DTYPE_Invalid:
 	case DTYPE_Preview:
-	case DTYPE_Procedural:
 	default:
 		UE_LOG( LogAndroidAudio, Warning, TEXT("Init Buffer on unsupported sound type name = %s type = %d"), *InWave->GetName(), int32(DecompressionType));
 		break;
@@ -224,6 +255,5 @@ bool FSLESSoundBuffer::ReadCompressedData( uint8* Destination, bool bLooping )
 */
 int FSLESSoundBuffer::GetRTBufferSize(void)
 {
-	check(DecompressionState != NULL);
-	return DecompressionState ? DecompressionState->GetStreamBufferSize() : 0;
+	return DecompressionState ? DecompressionState->GetStreamBufferSize() : MONO_PCM_BUFFER_SIZE;
 }
diff --git a/Engine/Source/Runtime/Android/AndroidAudio/Private/AndroidAudioSource.cpp b/Engine/Source/Runtime/Android/AndroidAudio/Private/AndroidAudioSource.cpp
--- a/Engine/Source/Runtime/Android/AndroidAudio/Private/AndroidAudioSource.cpp
+++ b/Engine/Source/Runtime/Android/AndroidAudio/Private/AndroidAudioSource.cpp
@@ -45,7 +45,7 @@ void FSLESSoundSource::OnRequeueBufferCallback( SLAndroidSimpleBufferQueueItf In
 
 		// Switch to the next buffer and decode for the next time the callback fires
 		BufferInUse = !BufferInUse;
-		bHasLooped = Buffer->ReadCompressedData( AudioBuffers[BufferInUse].AudioData, WaveInstance->LoopingMode != LOOP_Never );
+		bHasLooped = ReadMorePCMData(BufferInUse);
 	}
 }
 
@@ -129,6 +129,32 @@ bool FSLESSoundSource::EnqueuePCMBuffer( bool bLoop)
 	return true;
 }
 
+bool FSLESSoundSource::ReadProceduralData(const int32 BufferIndex)
+{
+	ensure(WaveInstance);
+	const int32 MaxSamples = BufferSize / sizeof(int16);
+	const int32 BytesWritten = WaveInstance->WaveData->GeneratePCMData(AudioBuffers[BufferIndex].AudioData, MaxSamples);
+	
+	AudioBuffers[BufferIndex].AudioDataSize = BytesWritten;
+
+	// convenience return value: we're never actually "looping" here.
+	return false;
+}
+
+bool FSLESSoundSource::ReadMorePCMData(const int32 BufferIndex)
+{
+	USoundWave *WaveData = WaveInstance->WaveData;
+	if (WaveData && WaveData->bProcedural)
+	{
+		return ReadProceduralData(BufferIndex);
+	}
+	else
+	{
+		return Buffer->ReadCompressedData(AudioBuffers[BufferIndex].AudioData, WaveInstance->LoopingMode != LOOP_Never);
+	}
+}
+
+
 bool FSLESSoundSource::EnqueuePCMRTBuffer( bool bLoop )
 {
 	if (AudioBuffers[0].AudioData || AudioBuffers[1].AudioData)
@@ -138,17 +164,17 @@ bool FSLESSoundSource::EnqueuePCMRTBuffer( bool bLoop )
 	FMemory::Memzero( AudioBuffers, sizeof( SLESAudioBuffer ) * 2 );
 
 	// Set up double buffer area to decompress to
-	const int bufferSize = Buffer->GetRTBufferSize() * Buffer->NumChannels;
+	BufferSize = Buffer->GetRTBufferSize() * Buffer->NumChannels;
 
-	AudioBuffers[0].AudioData = (uint8*)FMemory::Malloc(bufferSize);
-	AudioBuffers[0].AudioDataSize = bufferSize;
+	AudioBuffers[0].AudioData = (uint8*)FMemory::Malloc(BufferSize);
+	AudioBuffers[0].AudioDataSize = BufferSize;
 
-	AudioBuffers[1].AudioData = (uint8*)FMemory::Malloc(bufferSize);
-	AudioBuffers[1].AudioDataSize = bufferSize;
+	AudioBuffers[1].AudioData = (uint8*)FMemory::Malloc(BufferSize);
+	AudioBuffers[1].AudioDataSize = BufferSize;
 
 	// Decompress two buffers worth of data to fill up the queue
-	Buffer->ReadCompressedData( AudioBuffers[0].AudioData, bLoop );
-	Buffer->ReadCompressedData( AudioBuffers[1].AudioData, bLoop );
+	ReadMorePCMData(0);
+	ReadMorePCMData(1);
 
 	SLresult result;
 
@@ -199,6 +225,8 @@ bool FSLESSoundSource::Init( FWaveInstance* InWaveInstance )
 		UE_LOG( LogAndroidAudio, Warning, TEXT(" InitSoundSouce with PlayerObject not NULL, possible leak"));
 	}
 	
+	WaveInstance = InWaveInstance;
+
 	// Find matching buffer.
 	Buffer = FSLESSoundBuffer::Init( (FSLESAudioDevice *)AudioDevice, InWaveInstance->WaveData );
 
@@ -235,9 +263,6 @@ bool FSLESSoundSource::Init( FWaveInstance* InWaveInstance )
 			return false;
 		}
 		
-		WaveInstance = InWaveInstance;
-		
-		
 		Update();
 		
 		// Initialization was successful.
@@ -267,6 +292,7 @@ FSLESSoundSource::FSLESSoundSource( class FAudioDevice* InAudioDevice )
 		Buffer( NULL ),
 		bStreamedSound(false),
 		bBuffersToFlush(false),
+		BufferSize(0),
 		BufferInUse(0),
 		bHasLooped(false),
 		SL_PlayerObject(NULL),
diff --git a/Engine/Source/Runtime/Android/AndroidAudio/Public/AndroidAudioDevice.h b/Engine/Source/Runtime/Android/AndroidAudio/Public/AndroidAudioDevice.h
--- a/Engine/Source/Runtime/Android/AndroidAudio/Public/AndroidAudioDevice.h
+++ b/Engine/Source/Runtime/Android/AndroidAudio/Public/AndroidAudioDevice.h
@@ -56,7 +56,7 @@ public:
 	 *
 	 * @param InWave		USoundWave to use as template and wave source
 	 * @param AudioDevice	audio device to attach created buffer to
-	 * @return FCoreAudioSoundBuffer pointer if buffer creation succeeded, NULL otherwise
+	 * @return FSLESSoundBuffer pointer if buffer creation succeeded, NULL otherwise
 	 */
 	static FSLESSoundBuffer* CreateQueuedBuffer( FSLESAudioDevice* AudioDevice, USoundWave* Wave );
 
@@ -65,10 +65,19 @@ public:
 	 *
 	 * @param Wave			USoundWave to use as template and wave source
 	 * @param AudioDevice	audio device to attach created buffer to
-	 * @return FCoreAudioSoundBuffer pointer if buffer creation succeeded, NULL otherwise
+	 * @return FSLESSoundBuffer pointer if buffer creation succeeded, NULL otherwise
 	 */
 	static FSLESSoundBuffer* CreateNativeBuffer( FSLESAudioDevice* AudioDevice, USoundWave* Wave );
 
+	/**
+	* Static function used to create an Audio buffer and dynamically upload procedural data to.
+	*
+	* @param InWave		USoundWave to use as template and wave source
+	* @param AudioDevice	audio device to attach created buffer to
+	* @return FSLESSoundBuffer pointer if buffer creation succeeded, NULL otherwise
+	*/
+	static FSLESSoundBuffer* CreateProceduralBuffer(FSLESAudioDevice* AudioDevice, USoundWave* Wave);
+
 	/**
 	 * Static function used to create a buffer.
 	 *
@@ -206,6 +215,8 @@ protected:
 	/** Set when we wish to let the buffers play themselves out */
 	bool						bBuffersToFlush;
 
+	uint32						BufferSize;
+
 	int32						BufferInUse;
 	bool						bHasLooped;
 
@@ -214,6 +225,12 @@ protected:
 	void ReleaseResources();
 	bool EnqueuePCMBuffer( bool bLoop);
 	bool EnqueuePCMRTBuffer( bool bLoop);
+
+	/** Decompress through FSLESSoundBuffer, or call USoundWave procedure to generate more PCM data. Returns true/false: did audio loop? */
+	bool ReadMorePCMData(const int32 BufferIndex);
+
+	/** Handle obtaining more data for procedural USoundWaves. Always returns false for convenience. */
+	bool ReadProceduralData(const int32 BufferIndex);
 };
 
 /**
