Moving settings for audio compression control under the Android tab on the project properties.

diff --git a/Engine/Source/Developer/Android/AndroidTargetPlatform/Private/AndroidTargetPlatform.inl b/Engine/Source/Developer/Android/AndroidTargetPlatform/Private/AndroidTargetPlatform.inl
--- a/Engine/Source/Developer/Android/AndroidTargetPlatform/Private/AndroidTargetPlatform.inl
+++ b/Engine/Source/Developer/Android/AndroidTargetPlatform/Private/AndroidTargetPlatform.inl
@@ -283,7 +283,7 @@ FName FAndroidTargetPlatform<TPlatformProperties>::GetWaveFormat( class USoundWa
 		formatRead = true;
 
 		FString audioSetting;
-		if (!GConfig->GetString(TEXT("/Script/UnrealEd.CookerSettings"), TEXT("AndroidAudio"), audioSetting, GEngineIni))
+		if (!GConfig->GetString(TEXT("/Script/AndroidRuntimeSettings.AndroidRuntimeSettings"), TEXT("AndroidAudio"), audioSetting, GEngineIni))
 		{
 			audioSetting = TEXT("DEFAULT");
 		}
diff --git a/Engine/Source/Editor/UnrealEd/Private/Commandlets/CookCommandlet.cpp b/Engine/Source/Editor/UnrealEd/Private/Commandlets/CookCommandlet.cpp
--- a/Engine/Source/Editor/UnrealEd/Private/Commandlets/CookCommandlet.cpp
+++ b/Engine/Source/Editor/UnrealEd/Private/Commandlets/CookCommandlet.cpp
@@ -31,7 +31,6 @@ UCookerSettings::UCookerSettings(const FPostConstructInitializeProperties& PCIP)
 	: Super(PCIP)
 {
 	DefaultPVRTCQuality = 0;
-	AndroidAudio = EAndroidAudio::Default;
 }
 
 
diff --git a/Engine/Source/Editor/UnrealEd/Public/CookerSettings.h b/Engine/Source/Editor/UnrealEd/Public/CookerSettings.h
--- a/Engine/Source/Editor/UnrealEd/Public/CookerSettings.h
+++ b/Engine/Source/Editor/UnrealEd/Public/CookerSettings.h
@@ -8,18 +8,6 @@
 
 #include "CookerSettings.generated.h"
 
-
-UENUM()
-namespace EAndroidAudio
-{
-	enum Type
-	{
-		Default = 0 UMETA(DisplayName="Default",ToolTip="This option selects the default encoder."),
-		OGG = 1 UMETA(DisplayName="Ogg Vorbis", ToolTip="Selects Ogg Vorbis encoding."),
-		ADPCM = 2 UMETA(DisplayName="ADPCM",ToolTip="This option selects ADPCM lossless encoding.")
-	};
-}
-
 /**
  * Implements per-project cooker settings exposed to the editor
  */
@@ -34,8 +22,4 @@ public:
 	/** Quality of 0 means fastest, 4 means best */
 	UPROPERTY(config, EditAnywhere, Category=Textures)
 	int32 DefaultPVRTCQuality;
-
-	/** Android Audio encoding options */
-	UPROPERTY(config, EditAnywhere, Category = Android, meta = (DisplayName = "Audio encoding"))
-	TEnumAsByte<EAndroidAudio::Type> AndroidAudio;
 };
diff --git a/Engine/Source/Runtime/Android/AndroidRuntimeSettings/Classes/AndroidRuntimeSettings.h b/Engine/Source/Runtime/Android/AndroidRuntimeSettings/Classes/AndroidRuntimeSettings.h
--- a/Engine/Source/Runtime/Android/AndroidRuntimeSettings/Classes/AndroidRuntimeSettings.h
+++ b/Engine/Source/Runtime/Android/AndroidRuntimeSettings/Classes/AndroidRuntimeSettings.h
@@ -82,6 +82,17 @@ struct FGooglePlayLeaderboardMapping
 	FString LeaderboardID;
 };
 
+UENUM()
+namespace EAndroidAudio
+{
+	enum Type
+	{
+		Default = 0 UMETA(DisplayName = "Default", ToolTip = "This option selects the default encoder."),
+		OGG = 1 UMETA(DisplayName = "Ogg Vorbis", ToolTip = "Selects Ogg Vorbis encoding."),
+		ADPCM = 2 UMETA(DisplayName = "ADPCM", ToolTip = "This option selects ADPCM lossless encoding.")
+	};
+}
+
 /**
  * Implements the settings for the Android runtime platform.
  */
@@ -118,4 +129,8 @@ public:
 	// The unique identifier for the ad obtained from AdMob.
 	UPROPERTY(GlobalConfig, EditAnywhere, Category = GooglePlayServices)
 	FString AdMobAdUnitID;
+
+	/** Android Audio encoding options */
+	UPROPERTY(GlobalConfig, EditAnywhere, Category = DataCooker, meta = (DisplayName = "Audio encoding"))
+	TEnumAsByte<EAndroidAudio::Type> AndroidAudio;
 };
