- Improved tooltip/comment for x86

diff --git a/Engine/Source/Runtime/Android/AndroidRuntimeSettings/Classes/AndroidRuntimeSettings.h b/Engine/Source/Runtime/Android/AndroidRuntimeSettings/Classes/AndroidRuntimeSettings.h
--- a/Engine/Source/Runtime/Android/AndroidRuntimeSettings/Classes/AndroidRuntimeSettings.h
+++ b/Engine/Source/Runtime/Android/AndroidRuntimeSettings/Classes/AndroidRuntimeSettings.h
@@ -118,7 +118,7 @@ public:
 //	UPROPERTY(GlobalConfig, EditAnywhere, Category = Build, meta = (DisplayName = "Support arm64"))
 //	bool bBuildForArm64;
 
-	// Enable x86 support?
+	// Enable x86 support? [CURRENTLY FOR FULL SOURCE GAMES ONLY]
 	UPROPERTY(GlobalConfig, EditAnywhere, Category = Build, meta = (DisplayName = "Support x86"))
 	bool bBuildForX86;
 
@@ -126,7 +126,7 @@ public:
 //	UPROPERTY(GlobalConfig, EditAnywhere, Category = Build, meta = (DisplayName = "Support x86_64 [aka x64]"))
 //	bool bBuildForX8664;
 
-	// If selected, the checked architectures will be split into separate .apk files
+	// If selected, the checked architectures will be split into separate .apk files [CURRENTLY FOR FULL SOURCE GAMES ONLY]
 	UPROPERTY(GlobalConfig, EditAnywhere, Category = Build)
 	bool bSplitIntoSeparateApks;
 
