Changes to Android audio startup to resolve a crash bug a licensee was having.
Error reporting for creation failure changed from Fatal to Warning as we don't really need to stop there.

diff --git a/Engine/Source/Runtime/Android/AndroidAudio/Private/AndroidAudioSource.cpp b/Engine/Source/Runtime/Android/AndroidAudio/Private/AndroidAudioSource.cpp
--- a/Engine/Source/Runtime/Android/AndroidAudio/Private/AndroidAudioSource.cpp
+++ b/Engine/Source/Runtime/Android/AndroidAudio/Private/AndroidAudioSource.cpp
@@ -80,7 +80,7 @@ bool FSLESSoundSource::CreatePlayer()
 																	&SoundDataSource, &AudioSink, sizeof(ids) / sizeof(SLInterfaceID), ids, req );
 	if(result != SL_RESULT_SUCCESS)
 	{
-		UE_LOG( LogAndroidAudio, Fatal,TEXT("FAILED OPENSL BUFFER CreateAudioPlayer 0x%x"), result);
+		UE_LOG(LogAndroidAudio, Warning, TEXT("FAILED OPENSL BUFFER CreateAudioPlayer 0x%x"), result);
 		return false;
 	}
 		
@@ -88,18 +88,18 @@ bool FSLESSoundSource::CreatePlayer()
 		
 	// realize the player
 	result = (*SL_PlayerObject)->Realize(SL_PlayerObject, SL_BOOLEAN_FALSE);
-	if(result != SL_RESULT_SUCCESS) { UE_LOG( LogAndroidAudio, Fatal,TEXT("FAILED OPENSL BUFFER Realize 0x%x"), result); return false; }
+	if (result != SL_RESULT_SUCCESS) { UE_LOG(LogAndroidAudio, Warning, TEXT("FAILED OPENSL BUFFER Realize 0x%x"), result); return false; }
 		
 	// get the play interface
 	result = (*SL_PlayerObject)->GetInterface(SL_PlayerObject, SL_IID_PLAY, &SL_PlayerPlayInterface);
-	if(result != SL_RESULT_SUCCESS) { UE_LOG( LogAndroidAudio, Fatal,TEXT("FAILED OPENSL BUFFER GetInterface SL_IID_PLAY 0x%x"), result); bFailedSetup |= true; }
+	if (result != SL_RESULT_SUCCESS) { UE_LOG(LogAndroidAudio, Warning, TEXT("FAILED OPENSL BUFFER GetInterface SL_IID_PLAY 0x%x"), result); bFailedSetup |= true; }
 	// volume
 	result = (*SL_PlayerObject)->GetInterface(SL_PlayerObject, SL_IID_VOLUME, &SL_VolumeInterface);
-	if(result != SL_RESULT_SUCCESS) { UE_LOG( LogAndroidAudio, Fatal,TEXT("FAILED OPENSL BUFFER GetInterface SL_IID_VOLUME 0x%x"), result); bFailedSetup |= true; }
+	if (result != SL_RESULT_SUCCESS) { UE_LOG(LogAndroidAudio, Warning, TEXT("FAILED OPENSL BUFFER GetInterface SL_IID_VOLUME 0x%x"), result); bFailedSetup |= true; }
 	// buffer system
 	result = (*SL_PlayerObject)->GetInterface(SL_PlayerObject, SL_IID_BUFFERQUEUE, &SL_PlayerBufferQueue);
-	if(result != SL_RESULT_SUCCESS) { UE_LOG( LogAndroidAudio, Fatal,TEXT("FAILED OPENSL BUFFER GetInterface SL_IID_BUFFERQUEUE 0x%x"), result); bFailedSetup |= true; }
-
+	if (result != SL_RESULT_SUCCESS) { UE_LOG(LogAndroidAudio, Warning, TEXT("FAILED OPENSL BUFFER GetInterface SL_IID_BUFFERQUEUE 0x%x"), result); bFailedSetup |= true; }
+	
 	return bFailedSetup == false;
 }
 
@@ -123,11 +123,11 @@ bool FSLESSoundSource::EnqueuePCMBuffer( bool bLoop)
 	if( bLoop ) 
 	{
 		result = (*SL_PlayerBufferQueue)->RegisterCallback(SL_PlayerBufferQueue, OpenSLBufferQueueCallback, (void*)this);
-		if(result != SL_RESULT_SUCCESS) { UE_LOG( LogAndroidAudio, Fatal,TEXT("FAILED OPENSL BUFFER QUEUE RegisterCallback 0x%x "), result); return false;}
+		if (result != SL_RESULT_SUCCESS) { UE_LOG(LogAndroidAudio, Warning, TEXT("FAILED OPENSL BUFFER QUEUE RegisterCallback 0x%x "), result); return false; }
 	}
 	
 	result = (*SL_PlayerBufferQueue)->Enqueue(SL_PlayerBufferQueue, Buffer->AudioData, Buffer->GetSize() );
-	if(result != SL_RESULT_SUCCESS) { UE_LOG( LogAndroidAudio, Fatal,TEXT("FAILED OPENSL BUFFER Enqueue SL_PlayerBufferQueue 0x%x params( %p, %d)"), result, Buffer->AudioData, int32(Buffer->GetSize())); return false;}
+	if (result != SL_RESULT_SUCCESS) { UE_LOG(LogAndroidAudio, Warning, TEXT("FAILED OPENSL BUFFER Enqueue SL_PlayerBufferQueue 0x%x params( %p, %d)"), result, Buffer->AudioData, int32(Buffer->GetSize())); return false; }
 
 	bStreamedSound = false;
 	bHasLooped = false;
@@ -192,7 +192,7 @@ bool FSLESSoundSource::EnqueuePCMRTBuffer( bool bLoop )
 	if(result == SL_RESULT_SUCCESS) 
 	{
 		result = (*SL_PlayerBufferQueue)->Enqueue(SL_PlayerBufferQueue, AudioBuffers[0].AudioData, AudioBuffers[0].AudioDataSize );
-		if(result != SL_RESULT_SUCCESS) { UE_LOG( LogAndroidAudio, Fatal,TEXT("FAILED OPENSL BUFFER Enqueue SL_PlayerBufferQueue 0x%x params( %p, %d)"), result, Buffer->AudioData, int32(Buffer->GetSize())); return false;}
+		if (result != SL_RESULT_SUCCESS) { UE_LOG(LogAndroidAudio, Warning, TEXT("FAILED OPENSL BUFFER Enqueue SL_PlayerBufferQueue 0x%x params( %p, %d)"), result, Buffer->AudioData, int32(Buffer->GetSize())); return false; }
 	}
 	else
 	{
@@ -209,7 +209,7 @@ bool FSLESSoundSource::EnqueuePCMRTBuffer( bool bLoop )
 /**
  * Initializes a source with a given wave instance and prepares it for playback.
  *
- * @param	WaveInstance	wave instace being primed for playback
+ * @param	WaveInstance	wave instance being primed for playback
  * @return	TRUE if initialization was successful, FALSE otherwise
  */
 bool FSLESSoundSource::Init( FWaveInstance* InWaveInstance )
@@ -232,8 +232,6 @@ bool FSLESSoundSource::Init( FWaveInstance* InWaveInstance )
 		UE_LOG( LogAndroidAudio, Warning, TEXT(" InitSoundSouce with PlayerObject not NULL, possible leak"));
 	}
 	
-	WaveInstance = InWaveInstance;
-
 	// Find matching buffer.
 	Buffer = FSLESSoundBuffer::Init( (FSLESAudioDevice *)AudioDevice, InWaveInstance->WaveData );
 
@@ -245,6 +243,8 @@ bool FSLESSoundSource::Init( FWaveInstance* InWaveInstance )
 
 		if (CreatePlayer())
 		{
+			WaveInstance = InWaveInstance;
+
 			switch( Buffer->Format)
 			{
 				case SoundFormat_PCM:
@@ -383,7 +383,7 @@ void FSLESSoundSource::Update( void )
 		Location = FVector( 0.f, 0.f, 0.f );
 	}
 	
-	// Set volume & Pitch
+	// Set volume (Pitch changes are not supported on current Android platforms!)
 	// also Location & Velocity
 	
 	// Convert volume to millibels.
