- Fixed default Android enums [UE-8306]

diff --git a/Engine/Config/BaseEngine.ini b/Engine/Config/BaseEngine.ini
--- a/Engine/Config/BaseEngine.ini
+++ b/Engine/Config/BaseEngine.ini
@@ -2144,8 +2144,8 @@ bBuildForES31=false
 bBuildForGL4=false
 bSplitIntoSeparateApks=false
 bPackageDataInsideApk=false
-Orientation=landscape
-DepthBufferPreference=0
+Orientation=Landscape
+DepthBufferPreference=Default
 PackageName=com.YourCompany.[PROJECT]
 StoreVersion=1
 VersionDisplayName=1.0
diff --git a/Engine/Source/Developer/Android/AndroidPlatformEditor/Private/AndroidTargetSettingsCustomization.h b/Engine/Source/Developer/Android/AndroidPlatformEditor/Private/AndroidTargetSettingsCustomization.h
--- a/Engine/Source/Developer/Android/AndroidPlatformEditor/Private/AndroidTargetSettingsCustomization.h
+++ b/Engine/Source/Developer/Android/AndroidPlatformEditor/Private/AndroidTargetSettingsCustomization.h
@@ -60,11 +60,5 @@ private:
 	// Is the App ID string writable?
 	TAttribute<bool> SetupForGooglePlayAttribute;
 
-	// Converts an orientation enum to the associated string value
-	static FString OrientationToString(const EAndroidScreenOrientation::Type Orientation);
-
-	// Converts an depth preference enum to the associated string value
-	static FString DepthBufferPreferenceToString(const EAndroidDepthBufferPreference::Type DepthBufferPreference);
-
-	IDetailLayoutBuilder* SavedLayoutBuilder;
+ 	IDetailLayoutBuilder* SavedLayoutBuilder;
 };
diff --git a/Engine/Source/Runtime/Android/AndroidRuntimeSettings/Classes/AndroidRuntimeSettings.h b/Engine/Source/Runtime/Android/AndroidRuntimeSettings/Classes/AndroidRuntimeSettings.h
--- a/Engine/Source/Runtime/Android/AndroidRuntimeSettings/Classes/AndroidRuntimeSettings.h
+++ b/Engine/Source/Runtime/Android/AndroidRuntimeSettings/Classes/AndroidRuntimeSettings.h
@@ -7,32 +7,32 @@
 UENUM()
 namespace EAndroidScreenOrientation
 {
-	// IF THIS CHANGES, MAKE SURE TO UPDATE UEDeployAndroid.cs!
+	// IF THIS CHANGES, MAKE SURE TO UPDATE UEDeployAndroid.cs, ConvertOrientationIniValue()!
 	enum Type
 	{
 		// Portrait orientation (the display is taller than it is wide)
-		Portrait UMETA(ManifestValue = "portrait"),
+		Portrait,
 
 		// Portrait orientation rotated 180 degrees
-		ReversePortrait UMETA(ManifestValue = "reversePortrait"),
+		ReversePortrait,
 
 		// Use either portrait or reverse portrait orientation, based on the device orientation sensor
-		SensorPortrait UMETA(ManifestValue = "sensorPortrait"),
+		SensorPortrait,
 
 		// Landscape orientation (the display is wider than it is tall)
-		Landscape UMETA(ManifestValue = "landscape"),
+		Landscape,
 
 		// Landscape orientation rotated 180 degrees
-		ReverseLandscape UMETA(ManifestValue = "reverseLandscape"),
+		ReverseLandscape,
 
 		// Use either landscape or reverse landscape orientation, based on the device orientation sensor
-		SensorLandscape UMETA(ManifestValue = "sensorLandscape"),
+		SensorLandscape,
 
 		// Use any orientation the device normally supports, based on the device orientation sensor
-		Sensor UMETA(ManifestValue = "sensor"),
+		Sensor,
 
 		// Use any orientation (including ones the device wouldn't choose in Sensor mode), based on the device orientation sensor
-		FullSensor UMETA(ManifestValue = "fullSensor"),
+		FullSensor,
 	};
 }
 
@@ -40,13 +40,13 @@ namespace EAndroidScreenOrientation
 UENUM()
 namespace EAndroidDepthBufferPreference
 {
-	// IF THIS CHANGES, MAKE SURE TO UPDATE UEDeployAndroid.cs!
+	// IF THIS CHANGES, MAKE SURE TO UPDATE UEDeployAndroid.cs, ConvertDepthBufferIniValue()!
 	enum Type
 	{
-		Default = 0 UMETA(DisplayName = "Default", ManifestValue = "0"),
-		Bits16 = 16 UMETA(DisplayName = "16-bit", ManifestValue = "16"),
-		Bits24 = 24 UMETA(DisplayName = "24-bit", ManifestValue = "24"),
-		Bits32 = 32 UMETA(DisplayName = "32-bit", ManifestValue = "32"),
+		Default = 0 UMETA(DisplayName = "Default"),
+		Bits16 = 16 UMETA(DisplayName = "16-bit"),
+		Bits24 = 24 UMETA(DisplayName = "24-bit"),
+		Bits32 = 32 UMETA(DisplayName = "32-bit"),
 	};
 }
 
