UE-14954 Making android audio inaudible when linear volume of a sound is 0.0

diff --git a/Engine/Source/Runtime/Android/AndroidAudio/Private/AndroidAudioSource.cpp b/Engine/Source/Runtime/Android/AndroidAudio/Private/AndroidAudioSource.cpp
--- a/Engine/Source/Runtime/Android/AndroidAudio/Private/AndroidAudioSource.cpp
+++ b/Engine/Source/Runtime/Android/AndroidAudio/Private/AndroidAudioSource.cpp
@@ -367,7 +367,8 @@ void FSLESSoundSource::Update( void )
 		Volume = FMath::Clamp(Volume, 0.0f, MAX_VOLUME);
 	}
 	
-	const float Pitch = FMath::Clamp<float>(WaveInstance->Pitch, MIN_PITCH, MAX_PITCH);
+	// Pitch is *not* currently supported on android... but this is commented out rather than deleted for reference
+	// const float Pitch = FMath::Clamp<float>(WaveInstance->Pitch, MIN_PITCH, MAX_PITCH);
 
 	// Set whether to apply reverb
 	SetReverbApplied(true);
@@ -398,7 +399,7 @@ void FSLESSoundSource::Update( void )
 	
 	// Convert volume to millibels.
 	SLmillibel MaxMillibel = 0;
-	SLmillibel MinMillibel = -3000;
+	static const SLmillibel MinMillibel = -12000; // -120 dB will be inaudible
 	(*SL_VolumeInterface)->GetMaxVolumeLevel( SL_VolumeInterface, &MaxMillibel );
 	SLmillibel VolumeMillibel = (Volume * (MaxMillibel - MinMillibel)) + MinMillibel;
 	VolumeMillibel = FMath::Clamp(VolumeMillibel, MinMillibel, MaxMillibel);
