UE-14954 Changing the implementation to convert the linear volume to mdB and clamping within the device range.

This will result in a more platform consistent and accurate volume.
 of the sound source

diff --git a/Engine/Source/Runtime/Android/AndroidAudio/Private/AndroidAudioSource.cpp b/Engine/Source/Runtime/Android/AndroidAudio/Private/AndroidAudioSource.cpp
--- a/Engine/Source/Runtime/Android/AndroidAudio/Private/AndroidAudioSource.cpp
+++ b/Engine/Source/Runtime/Android/AndroidAudio/Private/AndroidAudioSource.cpp
@@ -396,14 +396,16 @@ void FSLESSoundSource::Update( void )
 	
 	// Set volume (Pitch changes are not supported on current Android platforms!)
 	// also Location & Velocity
-	
+
 	// Convert volume to millibels.
-	SLmillibel MaxMillibel = 0;
-	static const SLmillibel MinMillibel = -12000; // -120 dB will be inaudible
-	(*SL_VolumeInterface)->GetMaxVolumeLevel( SL_VolumeInterface, &MaxMillibel );
-	SLmillibel VolumeMillibel = (Volume * (MaxMillibel - MinMillibel)) + MinMillibel;
-	VolumeMillibel = FMath::Clamp(VolumeMillibel, MinMillibel, MaxMillibel);
-	
+	SLmillibel VolumeMillibel = -12000;
+	if (Volume > 0.0f)
+	{
+		SLmillibel MaxMillibel = 0;
+		(*SL_VolumeInterface)->GetMaxVolumeLevel(SL_VolumeInterface, &MaxMillibel);
+		VolumeMillibel = (SLmillibel)FMath::Clamp<int64>((int64)(2000.0f * log10f(Volume)), -10000, (int64)MaxMillibel);
+	}
+
 	SLresult result = (*SL_VolumeInterface)->SetVolumeLevel(SL_VolumeInterface, VolumeMillibel);
 	check(SL_RESULT_SUCCESS == result);
 
