UE-14954 Making the min volume a static const value

diff --git a/Engine/Source/Runtime/Android/AndroidAudio/Private/AndroidAudioSource.cpp b/Engine/Source/Runtime/Android/AndroidAudio/Private/AndroidAudioSource.cpp
--- a/Engine/Source/Runtime/Android/AndroidAudio/Private/AndroidAudioSource.cpp
+++ b/Engine/Source/Runtime/Android/AndroidAudio/Private/AndroidAudioSource.cpp
@@ -398,12 +398,13 @@ void FSLESSoundSource::Update( void )
 	// also Location & Velocity
 
 	// Convert volume to millibels.
-	SLmillibel VolumeMillibel = -12000;
+	static const int64 MinVolumeMillibel = -12000;
+	SLmillibel VolumeMillibel = (SLmillibel)MinVolumeMillibel;
 	if (Volume > 0.0f)
 	{
 		SLmillibel MaxMillibel = 0;
 		(*SL_VolumeInterface)->GetMaxVolumeLevel(SL_VolumeInterface, &MaxMillibel);
-		VolumeMillibel = (SLmillibel)FMath::Clamp<int64>((int64)(2000.0f * log10f(Volume)), -10000, (int64)MaxMillibel);
+		VolumeMillibel = (SLmillibel)FMath::Clamp<int64>((int64)(2000.0f * log10f(Volume)), MinVolumeMillibel, (int64)MaxMillibel);
 	}
 
 	SLresult result = (*SL_VolumeInterface)->SetVolumeLevel(SL_VolumeInterface, VolumeMillibel);
