UE-14954 only doing the log10 calculation if the volume changes from previous update

diff --git a/Engine/Source/Runtime/Android/AndroidAudio/Private/AndroidAudioSource.cpp b/Engine/Source/Runtime/Android/AndroidAudio/Private/AndroidAudioSource.cpp
--- a/Engine/Source/Runtime/Android/AndroidAudio/Private/AndroidAudioSource.cpp
+++ b/Engine/Source/Runtime/Android/AndroidAudio/Private/AndroidAudioSource.cpp
@@ -301,6 +301,7 @@ FSLESSoundSource::FSLESSoundSource( class FAudioDevice* InAudioDevice )
 		bBuffersToFlush(false),
 		BufferSize(0),
 		BufferInUse(0),
+		VolumePreviousUpdate(-1.0f),
 		bHasLooped(false),
 		SL_PlayerObject(NULL),
 		SL_PlayerPlayInterface(NULL),
@@ -400,11 +401,14 @@ void FSLESSoundSource::Update( void )
 	// Convert volume to millibels.
 	static const int64 MinVolumeMillibel = -12000;
 	SLmillibel VolumeMillibel = (SLmillibel)MinVolumeMillibel;
-	if (Volume > 0.0f)
+
+	// Avoid doing the log calculation each update by only doing it if the volume changed
+	if (Volume > 0.0f && Volume != VolumePreviousUpdate)
 	{
 		SLmillibel MaxMillibel = 0;
 		(*SL_VolumeInterface)->GetMaxVolumeLevel(SL_VolumeInterface, &MaxMillibel);
 		VolumeMillibel = (SLmillibel)FMath::Clamp<int64>((int64)(2000.0f * log10f(Volume)), MinVolumeMillibel, (int64)MaxMillibel);
+		VolumePreviousUpdate = Volume;
 	}
 
 	SLresult result = (*SL_VolumeInterface)->SetVolumeLevel(SL_VolumeInterface, VolumeMillibel);
diff --git a/Engine/Source/Runtime/Android/AndroidAudio/Public/AndroidAudioDevice.h b/Engine/Source/Runtime/Android/AndroidAudio/Public/AndroidAudioDevice.h
--- a/Engine/Source/Runtime/Android/AndroidAudio/Public/AndroidAudioDevice.h
+++ b/Engine/Source/Runtime/Android/AndroidAudio/Public/AndroidAudioDevice.h
@@ -218,6 +218,7 @@ protected:
 	uint32						BufferSize;
 
 	int32						BufferInUse;
+	float						VolumePreviousUpdate;
 	bool						bHasLooped;
 
 	bool CreatePlayer();
