UE-14954 Changing code to not call the SetVolumeLevel API function  unless the volume has changed.

diff --git a/Engine/Source/Runtime/Android/AndroidAudio/Private/AndroidAudioSource.cpp b/Engine/Source/Runtime/Android/AndroidAudio/Private/AndroidAudioSource.cpp
--- a/Engine/Source/Runtime/Android/AndroidAudio/Private/AndroidAudioSource.cpp
+++ b/Engine/Source/Runtime/Android/AndroidAudio/Private/AndroidAudioSource.cpp
@@ -398,22 +398,19 @@ void FSLESSoundSource::Update( void )
 	// Set volume (Pitch changes are not supported on current Android platforms!)
 	// also Location & Velocity
 
-	// Convert volume to millibels.
-	static const int64 MinVolumeMillibel = -12000;
-	SLmillibel VolumeMillibel = (SLmillibel)MinVolumeMillibel;
-
 	// Avoid doing the log calculation each update by only doing it if the volume changed
 	if (Volume > 0.0f && Volume != VolumePreviousUpdate)
 	{
+		// Convert volume to millibels.
+		static const int64 MinVolumeMillibel = -12000;
 		SLmillibel MaxMillibel = 0;
 		(*SL_VolumeInterface)->GetMaxVolumeLevel(SL_VolumeInterface, &MaxMillibel);
-		VolumeMillibel = (SLmillibel)FMath::Clamp<int64>((int64)(2000.0f * log10f(Volume)), MinVolumeMillibel, (int64)MaxMillibel);
+		SLmillibel VolumeMillibel = (SLmillibel)FMath::Clamp<int64>((int64)(2000.0f * log10f(Volume)), MinVolumeMillibel, (int64)MaxMillibel);
 		VolumePreviousUpdate = Volume;
+		SLresult result = (*SL_VolumeInterface)->SetVolumeLevel(SL_VolumeInterface, VolumeMillibel);
+		check(SL_RESULT_SUCCESS == result);
 	}
 
-	SLresult result = (*SL_VolumeInterface)->SetVolumeLevel(SL_VolumeInterface, VolumeMillibel);
-	check(SL_RESULT_SUCCESS == result);
-
 }
 
 /**
