UE-14954 Making sure we actually set a min dB volume if the volume of a sound is 0.0

diff --git a/Engine/Source/Runtime/Android/AndroidAudio/Private/AndroidAudioSource.cpp b/Engine/Source/Runtime/Android/AndroidAudio/Private/AndroidAudioSource.cpp
--- a/Engine/Source/Runtime/Android/AndroidAudio/Private/AndroidAudioSource.cpp
+++ b/Engine/Source/Runtime/Android/AndroidAudio/Private/AndroidAudioSource.cpp
@@ -399,16 +399,24 @@ void FSLESSoundSource::Update( void )
 	// also Location & Velocity
 
 	// Avoid doing the log calculation each update by only doing it if the volume changed
-	if (Volume > 0.0f && Volume != VolumePreviousUpdate)
+	if (Volume != VolumePreviousUpdate)
 	{
-		// Convert volume to millibels.
-		static const int64 MinVolumeMillibel = -12000;
-		SLmillibel MaxMillibel = 0;
-		(*SL_VolumeInterface)->GetMaxVolumeLevel(SL_VolumeInterface, &MaxMillibel);
-		SLmillibel VolumeMillibel = (SLmillibel)FMath::Clamp<int64>((int64)(2000.0f * log10f(Volume)), MinVolumeMillibel, (int64)MaxMillibel);
 		VolumePreviousUpdate = Volume;
-		SLresult result = (*SL_VolumeInterface)->SetVolumeLevel(SL_VolumeInterface, VolumeMillibel);
-		check(SL_RESULT_SUCCESS == result);
+		static const int64 MinVolumeMillibel = -12000;
+		if (Volume > 0.0f)
+		{
+			// Convert volume to millibels.
+			SLmillibel MaxMillibel = 0;
+			(*SL_VolumeInterface)->GetMaxVolumeLevel(SL_VolumeInterface, &MaxMillibel);
+			SLmillibel VolumeMillibel = (SLmillibel)FMath::Clamp<int64>((int64)(2000.0f * log10f(Volume)), MinVolumeMillibel, (int64)MaxMillibel);
+			SLresult result = (*SL_VolumeInterface)->SetVolumeLevel(SL_VolumeInterface, VolumeMillibel);
+			check(SL_RESULT_SUCCESS == result);
+		}
+		else
+		{
+			SLresult result = (*SL_VolumeInterface)->SetVolumeLevel(SL_VolumeInterface, MinVolumeMillibel);
+			check(SL_RESULT_SUCCESS == result);
+		}
 	}
 
 }
