UE-16362 Sound attenuation has no discernible falloff on Android

- Issue was variable shadowing on the Volume parameter introduced by CL 2477046

diff --git a/Engine/Source/Runtime/Android/AndroidAudio/Private/AndroidAudioSource.cpp b/Engine/Source/Runtime/Android/AndroidAudio/Private/AndroidAudioSource.cpp
--- a/Engine/Source/Runtime/Android/AndroidAudio/Private/AndroidAudioSource.cpp
+++ b/Engine/Source/Runtime/Android/AndroidAudio/Private/AndroidAudioSource.cpp
@@ -348,24 +348,14 @@ void FSLESSoundSource::Update( void )
 		return;
 	}
 	
-	float Volume = 1.0f;
-
-	// If the audio device is muted, then mute our sound source
-	if (AudioDevice->bIsDeviceMuted)
-	{
-		Volume = 0.0f;
-	}
-	else
+	float Volume = WaveInstance->Volume * WaveInstance->VolumeMultiplier;
+	if (SetStereoBleed())
 	{
-		float Volume = WaveInstance->Volume * WaveInstance->VolumeMultiplier;
-		if (SetStereoBleed())
-		{
-			// Emulate the bleed to rear speakers followed by stereo fold down
-			Volume *= 1.25f;
-		}
-		Volume *= FApp::GetVolumeMultiplier();
-		Volume = FMath::Clamp(Volume, 0.0f, MAX_VOLUME);
+		// Emulate the bleed to rear speakers followed by stereo fold down
+		Volume *= 1.25f;
 	}
+	Volume *= FApp::GetVolumeMultiplier();
+	Volume = FMath::Clamp(Volume, 0.0f, MAX_VOLUME);
 	
 	const float Pitch = FMath::Clamp<float>(WaveInstance->Pitch, MIN_PITCH, MAX_PITCH);
 
