Clean up some tooltips
#ue4
#android

diff --git a/Engine/Source/Runtime/Android/AndroidRuntimeSettings/Classes/AndroidRuntimeSettings.h b/Engine/Source/Runtime/Android/AndroidRuntimeSettings/Classes/AndroidRuntimeSettings.h
--- a/Engine/Source/Runtime/Android/AndroidRuntimeSettings/Classes/AndroidRuntimeSettings.h
+++ b/Engine/Source/Runtime/Android/AndroidRuntimeSettings/Classes/AndroidRuntimeSettings.h
@@ -134,7 +134,7 @@ public:
 	UPROPERTY(GlobalConfig, EditAnywhere, Category = APKPackaging, Meta = (DisplayName = "Disable verify OBB on first start/update."))
 	bool bDisableVerifyOBBOnStartUp;
 	
-	// The permitted orientation or orientations of the application on the device
+	// The permitted orientation of the application on the device
 	UPROPERTY(GlobalConfig, EditAnywhere, Category = APKPackaging)
 	TEnumAsByte<EAndroidScreenOrientation::Type> Orientation;
 
@@ -242,7 +242,7 @@ public:
 	UPROPERTY(GlobalConfig, EditAnywhere, Category = GooglePlayServices)
 	FString AdMobAdUnitID;
 
-	// The unique identifier for the ad obtained from AdMob.
+	// The unique identifier for this application (needed for IAP)
 	UPROPERTY(GlobalConfig, EditAnywhere, Category = GooglePlayServices)
 	FString GooglePlayLicenseKey;
 
