UE-17300 Strategy Game crashes during game play on Android device
- Added initial nulling of pointer which was causing the crash

diff --git a/Engine/Source/Runtime/Android/AndroidAudio/Private/AndroidAudioSource.cpp b/Engine/Source/Runtime/Android/AndroidAudio/Private/AndroidAudioSource.cpp
--- a/Engine/Source/Runtime/Android/AndroidAudio/Private/AndroidAudioSource.cpp
+++ b/Engine/Source/Runtime/Android/AndroidAudio/Private/AndroidAudioSource.cpp
@@ -347,7 +347,8 @@ FSLESSoundSource::FSLESSoundSource( class FAudioDevice* InAudioDevice )
 		SL_PlayerObject(NULL),
 		SL_PlayerPlayInterface(NULL),
 		SL_PlayerBufferQueue(NULL),
-		SL_VolumeInterface(NULL)
+		SL_VolumeInterface(NULL),
+		RealtimeAsyncTask(NULL)
 {
 	FMemory::Memzero( AudioBuffers, sizeof( AudioBuffers ) );
 }
