UE-14954 Fix for android volume

* Fix uses proper linear->dB volume calculation

diff --git a/Engine/Source/Runtime/Android/AndroidAudio/Private/AndroidAudioSource.cpp b/Engine/Source/Runtime/Android/AndroidAudio/Private/AndroidAudioSource.cpp
--- a/Engine/Source/Runtime/Android/AndroidAudio/Private/AndroidAudioSource.cpp
+++ b/Engine/Source/Runtime/Android/AndroidAudio/Private/AndroidAudioSource.cpp
@@ -343,6 +343,7 @@ FSLESSoundSource::FSLESSoundSource( class FAudioDevice* InAudioDevice )
 		bBuffersToFlush(false),
 		BufferSize(0),
 		BufferInUse(0),
+		VolumePreviousUpdate(-1.0f),
 		bHasLooped(false),
 		SL_PlayerObject(NULL),
 		SL_PlayerPlayInterface(NULL),
@@ -404,6 +405,7 @@ void FSLESSoundSource::Update( void )
 		Volume *= 1.25f;
 	}
 	Volume *= FApp::GetVolumeMultiplier();
+	Volume *= AudioDevice->PlatformAudioHeadroom;
 	Volume = FMath::Clamp(Volume, 0.0f, MAX_VOLUME);
 	
 	const float Pitch = FMath::Clamp<float>(WaveInstance->Pitch, MIN_PITCH, MAX_PITCH);
@@ -435,15 +437,26 @@ void FSLESSoundSource::Update( void )
 	// Set volume (Pitch changes are not supported on current Android platforms!)
 	// also Location & Velocity
 	
-	// Convert volume to millibels.
-	SLmillibel MaxMillibel = 0;
-	SLmillibel MinMillibel = -3000;
-	(*SL_VolumeInterface)->GetMaxVolumeLevel( SL_VolumeInterface, &MaxMillibel );
-	SLmillibel VolumeMillibel = (Volume * (MaxMillibel - MinMillibel)) + MinMillibel;
-	VolumeMillibel = FMath::Clamp(VolumeMillibel, MinMillibel, MaxMillibel);
-	
-	SLresult result = (*SL_VolumeInterface)->SetVolumeLevel(SL_VolumeInterface, VolumeMillibel);
-	check(SL_RESULT_SUCCESS == result);
+	// Avoid doing the log calculation each update by only doing it if the volume changed
+	if (Volume != VolumePreviousUpdate)
+	{
+		VolumePreviousUpdate = Volume;
+		static const int64 MinVolumeMillibel = -12000;
+		if (Volume > 0.0f)
+		{
+			// Convert volume to millibels.
+			SLmillibel MaxMillibel = 0;
+			(*SL_VolumeInterface)->GetMaxVolumeLevel(SL_VolumeInterface, &MaxMillibel);
+			SLmillibel VolumeMillibel = (SLmillibel)FMath::Clamp<int64>((int64)(2000.0f * log10f(Volume)), MinVolumeMillibel, (int64)MaxMillibel);
+			SLresult result = (*SL_VolumeInterface)->SetVolumeLevel(SL_VolumeInterface, VolumeMillibel);
+			check(SL_RESULT_SUCCESS == result);
+		}
+		else
+		{
+			SLresult result = (*SL_VolumeInterface)->SetVolumeLevel(SL_VolumeInterface, MinVolumeMillibel);
+			check(SL_RESULT_SUCCESS == result);
+		}
+	}
 
 }
 
