UE-18992 Fixing inconsistent volume on android

- Cached previous volume value needed to be reset on play so it has a chance of getting set in the API call for sound source instances which are recycled but playing the same wave instance.

diff --git a/Engine/Source/Runtime/Android/AndroidAudio/Private/AndroidAudioSource.cpp b/Engine/Source/Runtime/Android/AndroidAudio/Private/AndroidAudioSource.cpp
--- a/Engine/Source/Runtime/Android/AndroidAudio/Private/AndroidAudioSource.cpp
+++ b/Engine/Source/Runtime/Android/AndroidAudio/Private/AndroidAudioSource.cpp
@@ -467,7 +467,9 @@ void FSLESSoundSource::Play( void )
 {
 	if( WaveInstance )
 	{
-		
+		// Reset the previous volume on play so it can be set at least once in the update function
+		VolumePreviousUpdate = -1.0f;
+
 		// set the player's state to playing
 		SLresult result = (*SL_PlayerPlayInterface)->SetPlayState(SL_PlayerPlayInterface, SL_PLAYSTATE_PLAYING);
 		check(SL_RESULT_SUCCESS == result);
