UE-18731 - Sensor Portrait tooltip misleading
- Updated tool tip text as suggested by [REDACTED_PERSON]

diff --git a/Engine/Source/Runtime/Android/AndroidRuntimeSettings/Classes/AndroidRuntimeSettings.h b/Engine/Source/Runtime/Android/AndroidRuntimeSettings/Classes/AndroidRuntimeSettings.h
--- a/Engine/Source/Runtime/Android/AndroidRuntimeSettings/Classes/AndroidRuntimeSettings.h
+++ b/Engine/Source/Runtime/Android/AndroidRuntimeSettings/Classes/AndroidRuntimeSettings.h
@@ -17,7 +17,7 @@ namespace EAndroidScreenOrientation
 		/** Portrait orientation rotated 180 degrees. */
 		ReversePortrait,
 
-		/** Use either portrait or reverse portrait orientation, based on the device orientation sensor. */
+		/** Use either portrait or reverse portrait orientation, where supported by the device, based on the device orientation sensor. */
 		SensorPortrait,
 
 		/** Landscape orientation (the display is wider than it is tall). */
