Merging using UE4-To-UE4-Releases-4.10
Update Android minimum SDK level tooltip

diff --git a/Engine/Source/Runtime/Android/AndroidRuntimeSettings/Classes/AndroidRuntimeSettings.h b/Engine/Source/Runtime/Android/AndroidRuntimeSettings/Classes/AndroidRuntimeSettings.h
--- a/Engine/Source/Runtime/Android/AndroidRuntimeSettings/Classes/AndroidRuntimeSettings.h
+++ b/Engine/Source/Runtime/Android/AndroidRuntimeSettings/Classes/AndroidRuntimeSettings.h
@@ -122,7 +122,7 @@ public:
 	UPROPERTY(GlobalConfig, EditAnywhere, Category = APKPackaging, Meta = (DisplayName = "Version Display Name (usually x.y)"))
 	FString VersionDisplayName;
 
-	// What OS version the app is allowed to be installed on
+	// What OS version the app is allowed to be installed on (do not set this lower than 9)
 	UPROPERTY(GlobalConfig, EditAnywhere, Category = APKPackaging, Meta = (DisplayName = "Minimum SDK Version (9=Gingerbread, 14=Ice Cream Sandwich, 21=Lollipop)"))
 	int32 MinSDKVersion;
 	
