Merging using UE4-To-UE4-Releases-4.10
Update Android package name toolip

diff --git a/Engine/Source/Runtime/Android/AndroidRuntimeSettings/Classes/AndroidRuntimeSettings.h b/Engine/Source/Runtime/Android/AndroidRuntimeSettings/Classes/AndroidRuntimeSettings.h
--- a/Engine/Source/Runtime/Android/AndroidRuntimeSettings/Classes/AndroidRuntimeSettings.h
+++ b/Engine/Source/Runtime/Android/AndroidRuntimeSettings/Classes/AndroidRuntimeSettings.h
@@ -106,7 +106,7 @@ class ANDROIDRUNTIMESETTINGS_API UAndroidRuntimeSettings : public UObject
 public:
 	GENERATED_UCLASS_BODY()
 
-	// The official name of the product (same as the name you use on the Play Store web site)
+	// The official name of the product (same as the name you use on the Play Store web site). Note: Must have at least 2 sections separated by a period and be unique!
 	UPROPERTY(GlobalConfig, EditAnywhere, Category = APKPackaging, Meta = (DisplayName = "Android Package Name ('com.Company.Project', [PROJECT] is replaced with project name)"))
 	FString PackageName;
 
