UEMOB-454 Move the virtual keyboard out of Experimental and make it enabled by default

#4.21

diff --git a/Engine/Config/BaseEngine.ini b/Engine/Config/BaseEngine.ini
--- a/Engine/Config/BaseEngine.ini
+++ b/Engine/Config/BaseEngine.ini
@@ -458,6 +458,7 @@ bUseStreamingPause=false
 +FunctionRedirects=(OldName="SetSocialScreenTexture",NewName="/Script/HeadMountedDisplay.HeadMountedDisplayFunctionLibrary.SetSpectatorScreenTexture")
 +PropertyRedirects=(OldName="Widget.ClipToBounds",NewName="Clipping")
 
+
 ; 4.18
 +StructRedirects=(OldName="TargetReference",NewName="/Script/AnimGraphRuntime.BoneSocketTarget")
 
@@ -1747,7 +1748,7 @@ TextureFormatPriority_DXT=0.6
 TextureFormatPriority_PVRTC=0.8
 TextureFormatPriority_ATC=0.5
 TextureFormatPriority_ASTC=0.9
-bEnableNewKeyboard=False
+bEnableNewKeyboard=True
 bAndroidVoiceEnabled=false
 bBuildWithHiddenSymbolVisibility=false
 bSaveSymbols=false
diff --git a/Engine/Source/Runtime/Android/AndroidRuntimeSettings/Classes/AndroidRuntimeSettings.h b/Engine/Source/Runtime/Android/AndroidRuntimeSettings/Classes/AndroidRuntimeSettings.h
--- a/Engine/Source/Runtime/Android/AndroidRuntimeSettings/Classes/AndroidRuntimeSettings.h
+++ b/Engine/Source/Runtime/Android/AndroidRuntimeSettings/Classes/AndroidRuntimeSettings.h
@@ -249,7 +249,7 @@ public:
 	UPROPERTY(GlobalConfig, EditAnywhere, Category = "APK Packaging", Meta = (DisplayName = "Enable FullScreen Immersive on KitKat and above devices."))
 	bool bFullScreen;
 
-	UPROPERTY(GlobalConfig, EditAnywhere, Category = "APK Packaging", Meta = ( DisplayName = "Enable improved virtual keyboard [Experimental]"))
+	UPROPERTY(GlobalConfig, EditAnywhere, Category = "APK Packaging", Meta = ( DisplayName = "Enable improved virtual keyboard"))
 	bool bEnableNewKeyboard;
 	
 	// The preferred depth buffer bitcount for Android
