Added an Android Project Settings checkbox to enable detection of Vulkan device support by default. If unchecked, the -detectvulkan command line parameter can be used.

diff --git a/Engine/Config/BaseEngine.ini b/Engine/Config/BaseEngine.ini
--- a/Engine/Config/BaseEngine.ini
+++ b/Engine/Config/BaseEngine.ini
@@ -1754,6 +1754,7 @@ bBuildForX86=false
 bBuildForX8664=false
 bBuildForES2=true
 bSupportsVulkan=false
+bDetectVulkanByDefault=true
 bSplitIntoSeparateApks=false
 bPackageDataInsideApk=false
 bUseExternalFilesDir=false
diff --git a/Engine/Source/Runtime/Android/AndroidRuntimeSettings/Classes/AndroidRuntimeSettings.h b/Engine/Source/Runtime/Android/AndroidRuntimeSettings/Classes/AndroidRuntimeSettings.h
--- a/Engine/Source/Runtime/Android/AndroidRuntimeSettings/Classes/AndroidRuntimeSettings.h
+++ b/Engine/Source/Runtime/Android/AndroidRuntimeSettings/Classes/AndroidRuntimeSettings.h
@@ -348,6 +348,10 @@ public:
 	UPROPERTY(GlobalConfig, EditAnywhere, Category = Build, meta = (DisplayName = "Support Vulkan [Experimental]"))
 	bool bSupportsVulkan;
 
+	// Whether to detect Vulkan device support by default, if the project is packaged with Vulkan support. If unchecked, the -detectvulkan commandline will enable Vulkan detection.
+	UPROPERTY(GlobalConfig, EditAnywhere, AdvancedDisplay, Category = Build, meta = (DisplayName = "Detect Vulkan device support", EditCondition = "bSupportsVulkan"))
+	bool bDetectVulkanByDefault;
+
 	// Build the shipping config with hidden visibility by default. Results in smaller .so file but will also removes symbols used to display callstack dumps.
 	UPROPERTY(GlobalConfig, EditAnywhere, Category = AdvancedBuild, meta = (DisplayName = "Build with hidden symbol visibility in shipping config. [Experimental]"))
 	bool bBuildWithHiddenSymbolVisibility;
