Prevent Lumin from building the AndroidLocalNotification module. Fix CIS.

#rnx

diff --git a/Engine/Source/Runtime/Android/AndroidLocalNotification/AndroidLocalNotification.Build.cs b/Engine/Source/Runtime/Android/AndroidLocalNotification/AndroidLocalNotification.Build.cs
--- a/Engine/Source/Runtime/Android/AndroidLocalNotification/AndroidLocalNotification.Build.cs
+++ b/Engine/Source/Runtime/Android/AndroidLocalNotification/AndroidLocalNotification.Build.cs
@@ -14,7 +14,6 @@ public class AndroidLocalNotification : ModuleRules
             "Runtime/Engine/Public",
         });
 
-
         PrivateDependencyModuleNames.AddRange(new string[]
 		{
 			"Core",
@@ -22,5 +21,10 @@ public class AndroidLocalNotification : ModuleRules
 			"Engine",
             "Launch"
 		});
+
+		if(Target.Platform == UnrealTargetPlatform.Lumin)
+		{
+			PrecompileForTargets = PrecompileTargetsType.None;
+		}
 	}
 }
