Mark Enable Snapshots as experimental for now

#ue4
#android

diff --git a/Engine/Source/Runtime/Android/AndroidRuntimeSettings/Classes/AndroidRuntimeSettings.h b/Engine/Source/Runtime/Android/AndroidRuntimeSettings/Classes/AndroidRuntimeSettings.h
--- a/Engine/Source/Runtime/Android/AndroidRuntimeSettings/Classes/AndroidRuntimeSettings.h
+++ b/Engine/Source/Runtime/Android/AndroidRuntimeSettings/Classes/AndroidRuntimeSettings.h
@@ -393,7 +393,7 @@ public:
 	TArray<FGooglePlayLeaderboardMapping> LeaderboardMap;
 
 	// Enabling this requests snapshots support for saved games during Google Play login.
-	UPROPERTY(GlobalConfig, EditAnywhere, Category = GooglePlayServices, meta = (DisplayName = "Enable Snapshots on Google Play login"))
+	UPROPERTY(GlobalConfig, EditAnywhere, Category = GooglePlayServices, meta = (DisplayName = "Enable Snapshots on Google Play login [Experimental]"))
 	bool bEnableSnapshots;
 
 	// Enabling this includes the AdMob SDK and will be detected by Google Play Console on upload of APK.  Disable if you do not need ads to remove this warning.
