[NO NON-PII COMMIT MESSAGE]

diff --git a/Engine/Source/Runtime/Android/AudioMixerAndroid/Private/AudioMixerPlatformAndroid.cpp b/Engine/Source/Runtime/Android/AudioMixerAndroid/Private/AudioMixerPlatformAndroid.cpp
--- a/Engine/Source/Runtime/Android/AudioMixerAndroid/Private/AudioMixerPlatformAndroid.cpp
+++ b/Engine/Source/Runtime/Android/AudioMixerAndroid/Private/AudioMixerPlatformAndroid.cpp
@@ -51,6 +51,8 @@ namespace Audio
 		: bSuspended(false)
 		, bInitialized(false)
 		, bInCallback(false)
+		, NumSamplesPerRenderCallback(0)
+		, NumSamplesPerDeviceCallback(0)
 	{
 	}
 
@@ -87,6 +89,25 @@ namespace Audio
 		}
 	}
 
+	int32 FMixerPlatformAndroid::GetDeviceBufferSize(int32 RenderCallbackSize) const
+	{
+#if USE_ANDROID_JNI
+		// Override with platform-specific frames per buffer size
+		int32 MinFramesPerBuffer = AndroidThunkCpp_GetMetaDataInt(TEXT("audiomanager.framesPerBuffer"));
+
+		int32 BufferSizeToUse = MinFramesPerBuffer;
+		while (BufferSizeToUse < RenderCallbackSize)
+		{
+			BufferSizeToUse += MinFramesPerBuffer;
+		}
+
+		return BufferSizeToUse;
+#else
+		ensureMsgf(false, TEXT("JNI not supported on this platform. Audio output may be broken."));
+		return 1024;
+#endif
+	}
+
 	bool FMixerPlatformAndroid::InitializeHardware()
 	{
 		if (bInitialized)
@@ -209,7 +230,19 @@ namespace Audio
 
 		SLresult Result;
 
-		// data info
+		FAudioPlatformSettings PlatformSettings = GetPlatformSettings();
+
+		// Set up circular buffer between our rendering buffer size and the device's buffer size.
+		// Since we are only using this circular buffer on a single thread, we do not need to add extra slack.
+		NumSamplesPerRenderCallback = OpenStreamParams.NumFrames * AudioStreamInfo.DeviceInfo.NumChannels;
+		NumSamplesPerDeviceCallback = PlatformSettings.CallbackBufferFrameSize * AudioStreamInfo.DeviceInfo.NumChannels;
+		const int32 MaxCircularBufferCapacity = FMath::Max<int32>(NumSamplesPerRenderCallback, NumSamplesPerDeviceCallback) * 2;
+		CircularOutputBuffer.SetCapacity(MaxCircularBufferCapacity);
+
+		DeviceBuffer.Reset();
+		DeviceBuffer.AddUninitialized(NumSamplesPerDeviceCallback);
+
+		// Data Info:
 		SLDataLocator_AndroidSimpleBufferQueue LocationBuffer = { SL_DATALOCATOR_ANDROIDSIMPLEBUFFERQUEUE, 1};
 
 		// PCM Info
@@ -320,20 +353,7 @@ namespace Audio
  	{
 		FAudioPlatformSettings PlatformSettings = FAudioPlatformSettings::GetPlatformSettings(TEXT("/Script/AndroidRuntimeSettings.AndroidRuntimeSettings"));
 
-#if USE_ANDROID_JNI
-		// Override with platform-specific frames per buffer size
-		int32 MinFramesPerBuffer = AndroidThunkCpp_GetMetaDataInt(TEXT("audiomanager.framesPerBuffer"));
-
-		int32 BufferSizeToUse = MinFramesPerBuffer;
-		while (BufferSizeToUse < PlatformSettings.CallbackBufferFrameSize)
-		{
-			BufferSizeToUse += MinFramesPerBuffer;
-		}
-
-		PlatformSettings.CallbackBufferFrameSize = BufferSizeToUse;
-#else
-		// @todo Lumin: implement this?
-#endif
+		PlatformSettings.CallbackBufferFrameSize = GetDeviceBufferSize(PlatformSettings.CallbackBufferFrameSize);
 		return PlatformSettings;
 	}
 
@@ -376,9 +396,20 @@ namespace Audio
 
 	void FMixerPlatformAndroid::SubmitBuffer(const uint8* Buffer)
 	{
-		const auto BufferSize = AudioStreamInfo.NumOutputFrames * AudioStreamInfo.DeviceInfo.NumChannels * sizeof(int16);
-		SLresult Result = (*SL_PlayerBufferQueue)->Enqueue(SL_PlayerBufferQueue, Buffer, BufferSize);
-		OPENSLES_LOG_ON_FAIL(Result);
+		check(DeviceBuffer.Num() == NumSamplesPerDeviceCallback);
+
+		int32 PushResult = CircularOutputBuffer.Push((const int16*)Buffer, NumSamplesPerRenderCallback);
+		check(PushResult == NumSamplesPerRenderCallback)
+
+		while (CircularOutputBuffer.Num() >= NumSamplesPerDeviceCallback)
+		{
+			int32 PopResult = CircularOutputBuffer.Pop(DeviceBuffer.GetData(), NumSamplesPerDeviceCallback);
+			check(PopResult == NumSamplesPerDeviceCallback);
+
+			const auto BufferSize = NumSamplesPerDeviceCallback * sizeof(int16);
+			SLresult Result = (*SL_PlayerBufferQueue)->Enqueue(SL_PlayerBufferQueue, Buffer, BufferSize);
+			OPENSLES_LOG_ON_FAIL(Result);
+		}
 	}
 
 	FName FMixerPlatformAndroid::GetRuntimeFormat(USoundWave* InSoundWave)
diff --git a/Engine/Source/Runtime/Android/AudioMixerAndroid/Private/AudioMixerPlatformAndroid.h b/Engine/Source/Runtime/Android/AudioMixerAndroid/Private/AudioMixerPlatformAndroid.h
--- a/Engine/Source/Runtime/Android/AudioMixerAndroid/Private/AudioMixerPlatformAndroid.h
+++ b/Engine/Source/Runtime/Android/AudioMixerAndroid/Private/AudioMixerPlatformAndroid.h
@@ -3,6 +3,7 @@
 #pragma once
 
 #include "AudioMixer.h"
+#include "DSP/Dsp.h"
 #include <SLES/OpenSLES.h>
 #include "SLES/OpenSLES_Android.h"
 
@@ -46,6 +47,8 @@ namespace Audio
 	private:
 		const TCHAR* GetErrorString(SLresult Result);
 
+		int32 GetDeviceBufferSize(int32 RenderCallbackSize) const;
+
 		SLObjectItf	SL_EngineObject;
 		SLEngineItf	SL_EngineEngine;
 		SLObjectItf	SL_OutputMixObject;
@@ -59,6 +62,16 @@ namespace Audio
 		bool bInitialized;
 		bool bInCallback;
 		
+		// This buffer is pushed to and popped from in the SubmitBuffer callback. 
+		// This is required for devices that require frame counts per callback that are not powers of two.
+		Audio::TCircularAudioBuffer<int16> CircularOutputBuffer;
+
+		// This is the buffer we pop CircularOutputBuffer into in SubmitBuffer.
+		TArray<int16> DeviceBuffer;
+
+		int32 NumSamplesPerRenderCallback;
+		int32 NumSamplesPerDeviceCallback;
+
 		static void OpenSLBufferQueueCallback( SLAndroidSimpleBufferQueueItf InQueueInterface, void* pContext );		
 	};
 
