Fixed CancelLocalNotification which was not working on Android.
#rnx

diff --git a/Engine/Source/Runtime/Android/AndroidLocalNotification/Private/AndroidLocalNotification.cpp b/Engine/Source/Runtime/Android/AndroidLocalNotification/Private/AndroidLocalNotification.cpp
--- a/Engine/Source/Runtime/Android/AndroidLocalNotification/Private/AndroidLocalNotification.cpp
+++ b/Engine/Source/Runtime/Android/AndroidLocalNotification/Private/AndroidLocalNotification.cpp
@@ -126,6 +126,6 @@ void FAndroidLocalNotificationService::CancelLocalNotification(int32 Notificatio
 	
 #if USE_ANDROID_JNI
 	extern bool AndroidThunkCpp_DestroyScheduledNotificationIfExists(int32 NotificationId);
-	bool AndroidThunkCpp_DestroyScheduledNotificationIfExists(int32 NotificationId);
+	AndroidThunkCpp_DestroyScheduledNotificationIfExists(NotificationId);
 #endif
 }
