UE4.21 - Vulkan on Android is no longer experimental

diff --git a/Engine/Source/Runtime/Android/AndroidRuntimeSettings/Classes/AndroidRuntimeSettings.h b/Engine/Source/Runtime/Android/AndroidRuntimeSettings/Classes/AndroidRuntimeSettings.h
--- a/Engine/Source/Runtime/Android/AndroidRuntimeSettings/Classes/AndroidRuntimeSettings.h
+++ b/Engine/Source/Runtime/Android/AndroidRuntimeSettings/Classes/AndroidRuntimeSettings.h
@@ -350,7 +350,7 @@ public:
 	bool bBuildForES31;
 
 	// Enable Vulkan rendering support?
-	UPROPERTY(GlobalConfig, EditAnywhere, Category = Build, meta = (DisplayName = "Support Vulkan [Experimental]"))
+	UPROPERTY(GlobalConfig, EditAnywhere, Category = Build, meta = (DisplayName = "Support Vulkan"))
 	bool bSupportsVulkan;
 
 	// Whether to detect Vulkan device support by default, if the project is packaged with Vulkan support. If unchecked, the -detectvulkan commandline will enable Vulkan detection.
