Remove x86 option for Android (no longer supported)

#ue4
#android

diff --git a/Engine/Source/Runtime/Android/AndroidRuntimeSettings/Classes/AndroidRuntimeSettings.h b/Engine/Source/Runtime/Android/AndroidRuntimeSettings/Classes/AndroidRuntimeSettings.h
--- a/Engine/Source/Runtime/Android/AndroidRuntimeSettings/Classes/AndroidRuntimeSettings.h
+++ b/Engine/Source/Runtime/Android/AndroidRuntimeSettings/Classes/AndroidRuntimeSettings.h
@@ -333,10 +333,6 @@ public:
 	UPROPERTY(GlobalConfig, EditAnywhere, Category = Build, meta = (DisplayName = "Support arm64 [aka arm64-v8a]"))
 	bool bBuildForArm64;
 
-	// Enable x86 support? [CURRENTLY FOR FULL SOURCE GAMES ONLY]
-	UPROPERTY(GlobalConfig, EditAnywhere, Category = Build, meta = (DisplayName = "Support x86"))
-	bool bBuildForX86;
-
 	// Enable x86-64 support? [CURRENTLY FOR FULL SOURCE GAMES ONLY]
 	UPROPERTY(GlobalConfig, EditAnywhere, Category = Build, meta = (DisplayName = "Support x86_64 [aka x64]"))
 	bool bBuildForX8664;
