Update APK size limit tooltip (contributed by [REDACTED_PERSON])

#PR #4708
#ue4
#android

diff --git a/Engine/Source/Runtime/Android/AndroidRuntimeSettings/Classes/AndroidRuntimeSettings.h b/Engine/Source/Runtime/Android/AndroidRuntimeSettings/Classes/AndroidRuntimeSettings.h
--- a/Engine/Source/Runtime/Android/AndroidRuntimeSettings/Classes/AndroidRuntimeSettings.h
+++ b/Engine/Source/Runtime/Android/AndroidRuntimeSettings/Classes/AndroidRuntimeSettings.h
@@ -216,7 +216,7 @@ public:
 	UPROPERTY(GlobalConfig, EditAnywhere, Category = "APK Packaging", Meta = (DisplayName = "Enable Lint depreciation checks"))
 	bool bEnableLint;
 
-	// Should the data be placed into the .apk file instead of a separate .obb file. Amazon requires this to be enabled, but Google Play Store will not allow .apk files larger than 50MB, so only small games will work with this enabled.
+	// Should the data be placed into the .apk file instead of a separate .obb file. Amazon requires this to be enabled, but Google Play Store will not allow .apk files larger than 100MB, so only small games will work with this enabled.
 	UPROPERTY(GlobalConfig, EditAnywhere, Category = "APK Packaging", Meta = (DisplayName = "Package game data inside .apk?"))
 	bool bPackageDataInsideApk;
 
