[NO NON-PII COMMIT MESSAGE]

diff --git a/Engine/Source/Runtime/Android/AudioMixerAndroid/Private/AudioMixerPlatformAndroid.cpp b/Engine/Source/Runtime/Android/AudioMixerAndroid/Private/AudioMixerPlatformAndroid.cpp
--- a/Engine/Source/Runtime/Android/AudioMixerAndroid/Private/AudioMixerPlatformAndroid.cpp
+++ b/Engine/Source/Runtime/Android/AudioMixerAndroid/Private/AudioMixerPlatformAndroid.cpp
@@ -220,7 +220,7 @@ namespace Audio
 
 		AudioStreamInfo.OutputDeviceIndex = 0;
 		AudioStreamInfo.NumOutputFrames = OpenStreamParams.NumFrames;
-		AudioStreamInfo.NumBuffers = OpenStreamParams.NumBuffers;
+		AudioStreamInfo.NumBuffers = FMath::Max(OpenStreamParams.NumBuffers, 4);
 		AudioStreamInfo.AudioMixer = OpenStreamParams.AudioMixer;
 
 		if (!GetOutputDeviceInfo(AudioStreamInfo.OutputDeviceIndex, AudioStreamInfo.DeviceInfo))
