More platform fix-ups from prior max channel updates.

diff --git a/Engine/Source/Runtime/Android/AndroidAudio/Private/AndroidAudioDevice.cpp b/Engine/Source/Runtime/Android/AndroidAudio/Private/AndroidAudioDevice.cpp
--- a/Engine/Source/Runtime/Android/AndroidAudio/Private/AndroidAudioDevice.cpp
+++ b/Engine/Source/Runtime/Android/AndroidAudio/Private/AndroidAudioDevice.cpp
@@ -110,14 +110,14 @@ bool FSLESAudioDevice::InitializeHardware( void )
 
 
 	// Default to sensible channel count.
-	if( MaxChannels < 1 )
+	if(GetMaxChannels() < 1)
 	{
-		MaxChannels = 12;
+		SetMaxChannels(12);
 	}
 
 
 	// Initialize channels.
-	for( int32 i = 0; i < FMath::Min( MaxChannels, 12 ); i++ )
+	for( int32 i = 0; i < FMath::Min(GetMaxChannels(), 12 ); i++ )
 	{
 		FSLESSoundSource* Source = new FSLESSoundSource( this );
 		Sources.Add( Source );
@@ -131,8 +131,8 @@ bool FSLESAudioDevice::InitializeHardware( void )
 	}
 
 	// Update MaxChannels in case we couldn't create enough sources.
-	MaxChannels = Sources.Num();
-	UE_LOG( LogAndroidAudio, Warning, TEXT( "OpenSLAudioDevice: Allocated %i sources" ), MaxChannels );
+	SetMaxChannels(Sources.Num());
+	UE_LOG( LogAndroidAudio, Warning, TEXT( "OpenSLAudioDevice: Allocated %i sources" ), GetMaxChannels());
 
 	// Set up a default (nop) effects manager
 	Effects = new FAudioEffectsManager( this );
