More clean-up around initialization of audio MaxChannels/GetPlatformSettings
- Add GetMaxSources() that provides actual max number of sources (internal MaxSources + NumStoppingSources)
- More clean-up around platform implementation for GetPlatformSettings

diff --git a/Engine/Source/Editor/GameProjectGeneration/Private/GameProjectUtils.cpp b/Engine/Source/Editor/GameProjectGeneration/Private/GameProjectUtils.cpp
--- a/Engine/Source/Editor/GameProjectGeneration/Private/GameProjectUtils.cpp
+++ b/Engine/Source/Editor/GameProjectGeneration/Private/GameProjectUtils.cpp
@@ -138,9 +138,13 @@ namespace
 		DefaultProjectSettings.Add(TEXT("Mac"), MacSettings);
 
 		FAudioDefaultPlatformSettings PS4Settings(EAudioPlatform::Playstation4);
+		PS4Settings.Settings.CallbackBufferFrameSize = 256;
+		PS4Settings.Settings.NumBuffers = 7;
+		PS4Settings.Settings.NumSourceWorkers = 4;
 		DefaultProjectSettings.Add(TEXT("PS4"), PS4Settings);
 
 		FAudioDefaultPlatformSettings SwitchSettings(EAudioPlatform::Switch);
+		SwitchSettings.Settings.MaxChannels = 16;
 		DefaultProjectSettings.Add(TEXT("Switch"), SwitchSettings);
 
 		FAudioDefaultPlatformSettings WindowsSettings(EAudioPlatform::Windows);
diff --git a/Engine/Source/Runtime/Android/AndroidAudio/Private/AndroidAudioDevice.cpp b/Engine/Source/Runtime/Android/AndroidAudio/Private/AndroidAudioDevice.cpp
--- a/Engine/Source/Runtime/Android/AndroidAudio/Private/AndroidAudioDevice.cpp
+++ b/Engine/Source/Runtime/Android/AndroidAudio/Private/AndroidAudioDevice.cpp
@@ -8,6 +8,7 @@
 #include "AndroidAudioDevice.h"
 #include "VorbisAudioInfo.h"
 #include "ADPCMAudioInfo.h"
+#include "AudioPluginUtilities.h"
 
 DEFINE_LOG_CATEGORY(LogAndroidAudio);
 
@@ -71,71 +72,54 @@ void FSLESAudioDevice::Teardown( void )
  *
  * @return TRUE if initialization was successful, FALSE otherwise
  */
-bool FSLESAudioDevice::InitializeHardware( void )
+bool FSLESAudioDevice::InitializeHardware(void)
 {
-	UE_LOG( LogAndroidAudio, Warning, TEXT("SL Entered Init HW"));
+	UE_LOG(LogAndroidAudio, Warning, TEXT("SL Entered Init HW"));
 
 	SLresult result;
 
-	UE_LOG( LogAndroidAudio, Warning, TEXT("OpenSLES Initializing HW"));
+	UE_LOG(LogAndroidAudio, Warning, TEXT("OpenSLES Initializing HW"));
 
-	SLEngineOption EngineOption[] = { {(SLuint32) SL_ENGINEOPTION_THREADSAFE, (SLuint32) SL_BOOLEAN_TRUE} };
+	SLEngineOption EngineOption[] = { {(SLuint32)SL_ENGINEOPTION_THREADSAFE, (SLuint32)SL_BOOLEAN_TRUE} };
 
-    // create engine
-    result = slCreateEngine( &SL_EngineObject, 1, EngineOption, 0, NULL, NULL);
-    //check(SL_RESULT_SUCCESS == result);
+	// create engine
+	result = slCreateEngine(&SL_EngineObject, 1, EngineOption, 0, NULL, NULL);
+	//check(SL_RESULT_SUCCESS == result);
 	if (SL_RESULT_SUCCESS != result)
 	{
-		UE_LOG( LogAndroidAudio, Error, TEXT("Engine create failed %d"), int32(result));
+		UE_LOG(LogAndroidAudio, Error, TEXT("Engine create failed %d"), int32(result));
 	}
 
-    // realize the engine
-    result = (*SL_EngineObject)->Realize(SL_EngineObject, SL_BOOLEAN_FALSE);
-    check(SL_RESULT_SUCCESS == result);
+	// realize the engine
+	result = (*SL_EngineObject)->Realize(SL_EngineObject, SL_BOOLEAN_FALSE);
+	check(SL_RESULT_SUCCESS == result);
 
-    // get the engine interface, which is needed in order to create other objects
-    result = (*SL_EngineObject)->GetInterface(SL_EngineObject, SL_IID_ENGINE, &SL_EngineEngine);
-    check(SL_RESULT_SUCCESS == result);
+	// get the engine interface, which is needed in order to create other objects
+	result = (*SL_EngineObject)->GetInterface(SL_EngineObject, SL_IID_ENGINE, &SL_EngineEngine);
+	check(SL_RESULT_SUCCESS == result);
 
 	// create output mix, with environmental reverb specified as a non-required interface
-    result = (*SL_EngineEngine)->CreateOutputMix( SL_EngineEngine, &SL_OutputMixObject, 0, NULL, NULL );
-    check(SL_RESULT_SUCCESS == result);
+	result = (*SL_EngineEngine)->CreateOutputMix(SL_EngineEngine, &SL_OutputMixObject, 0, NULL, NULL);
+	check(SL_RESULT_SUCCESS == result);
 
-    // realize the output mix
-    result = (*SL_OutputMixObject)->Realize(SL_OutputMixObject, SL_BOOLEAN_FALSE);
-    check(SL_RESULT_SUCCESS == result);
+	// realize the output mix
+	result = (*SL_OutputMixObject)->Realize(SL_OutputMixObject, SL_BOOLEAN_FALSE);
+	check(SL_RESULT_SUCCESS == result);
 
-	UE_LOG( LogAndroidAudio, Warning, TEXT("OpenSLES Initialized"));
-    // ignore unsuccessful result codes for env
+	UE_LOG(LogAndroidAudio, Warning, TEXT("OpenSLES Initialized"));
 
-
-	// Default to sensible channel count.
-	if(GetMaxChannels() < 1)
-	{
-		SetMaxChannels(12);
-	}
-
-
-	// Initialize channels.
-	for( int32 i = 0; i < FMath::Min(GetMaxChannels(), 12 ); i++ )
+	// Initialize sources
+	const int32 MaxSoundSources = GetMaxSources();
+	for (int32 i = 0; i < MaxSoundSources; i++)
 	{
-		FSLESSoundSource* Source = new FSLESSoundSource( this );
-		Sources.Add( Source );
-		FreeSources.Add( Source );
+		FSLESSoundSource* Source = new FSLESSoundSource(this);
+		Sources.Add(Source);
+		FreeSources.Add(Source);
 	}
-
-	if( Sources.Num() < 1 )
-	{
-		UE_LOG( LogAndroidAudio,  Warning, TEXT( "OpenSLAudio: couldn't allocate any sources" ) );
-		return false;
-	}
-
-	// Update MaxChannels in case we couldn't create enough sources.
-	SetMaxChannels(Sources.Num());
-	UE_LOG( LogAndroidAudio, Warning, TEXT( "OpenSLAudioDevice: Allocated %i sources" ), GetMaxChannels());
+	UE_LOG(LogAndroidAudio, Warning, TEXT("OpenSLAudioDevice: Allocated %i sources"), MaxSoundSources);
 
 	// Set up a default (nop) effects manager
-	Effects = new FAudioEffectsManager( this );
+	Effects = new FAudioEffectsManager(this);
 
 	return true;
 }
@@ -200,7 +184,8 @@ class ICompressedAudioInfo* FSLESAudioDevice::CreateCompressedAudioInfo(USoundWa
 
 FAudioPlatformSettings FSLESAudioDevice::GetPlatformSettings() const
 {
-	return FAudioPlatformSettings::GetPlatformSettings(TEXT("/Script/AndroidRuntimeSettings.AndroidRuntimeSettings"));
+	const TCHAR* ConfigSection = AudioPluginUtilities::GetPlatformConfigSection(EAudioPlatform::Android);
+	return FAudioPlatformSettings::GetPlatformSettings(ConfigSection);
 }
 
 /** Check if any background music or sound is playing through the audio device */
diff --git a/Engine/Source/Runtime/Android/AudioMixerAndroid/Private/AudioMixerPlatformAndroid.cpp b/Engine/Source/Runtime/Android/AudioMixerAndroid/Private/AudioMixerPlatformAndroid.cpp
--- a/Engine/Source/Runtime/Android/AudioMixerAndroid/Private/AudioMixerPlatformAndroid.cpp
+++ b/Engine/Source/Runtime/Android/AudioMixerAndroid/Private/AudioMixerPlatformAndroid.cpp
@@ -8,6 +8,7 @@
 #include "Misc/ConfigCacheIni.h"
 #include "VorbisAudioInfo.h"
 #include "ADPCMAudioInfo.h"
+#include "AudioPluginUtilities.h"
 
 #include <SLES/OpenSLES.h>
 #include "SLES/OpenSLES_Android.h"
@@ -351,7 +352,8 @@ namespace Audio
 
  	FAudioPlatformSettings FMixerPlatformAndroid::GetPlatformSettings() const
  	{
-		FAudioPlatformSettings PlatformSettings = FAudioPlatformSettings::GetPlatformSettings(TEXT("/Script/AndroidRuntimeSettings.AndroidRuntimeSettings"));
+		const TCHAR* ConfigSection = AudioPluginUtilities::GetPlatformConfigSection(EAudioPlatform::Android);
+		FAudioPlatformSettings PlatformSettings = FAudioPlatformSettings::GetPlatformSettings(ConfigSection);
 
 		PlatformSettings.CallbackBufferFrameSize = GetDeviceBufferSize(PlatformSettings.CallbackBufferFrameSize);
 		return PlatformSettings;
diff --git a/Engine/Source/Runtime/AudioMixer/Private/AudioMixerDevice.cpp b/Engine/Source/Runtime/AudioMixer/Private/AudioMixerDevice.cpp
--- a/Engine/Source/Runtime/AudioMixer/Private/AudioMixerDevice.cpp
+++ b/Engine/Source/Runtime/AudioMixer/Private/AudioMixerDevice.cpp
@@ -162,7 +162,7 @@ namespace Audio
 			OpenStreamParams.OutputDeviceIndex = AUDIO_MIXER_DEFAULT_DEVICE_INDEX; // TODO: Support overriding which audio device user wants to open, not necessarily default.
 			OpenStreamParams.SampleRate = SampleRate;
 			OpenStreamParams.AudioMixer = this;
-			OpenStreamParams.MaxChannels = GetMaxChannels();
+			OpenStreamParams.MaxSources = GetMaxSources();
 
 			FString DefaultDeviceName = AudioMixerPlatform->GetDefaultDeviceName();
 
@@ -203,11 +203,10 @@ namespace Audio
 				// Initialize some data that depends on speaker configuration, etc.
 				InitializeChannelAzimuthMap(PlatformInfo.NumChannels);
 
-				// We initialize the number of sources to be 2 times the max channels
-				// This extra source count is used for "stopping sources", which are sources
-				// which are fading out (very quickly) to avoid discontinuities when stopping sounds
 				FSourceManagerInitParams SourceManagerInitParams;
-				SourceManagerInitParams.NumSources = GetMaxChannels() + NumStoppingVoices;
+				SourceManagerInitParams.NumSources = GetMaxSources();
+
+				// TODO: Migrate this to project settings properly
 				SourceManagerInitParams.NumSourceWorkers = 4;
 
 				SourceManager.Init(SourceManagerInitParams);
diff --git a/Engine/Source/Runtime/AudioMixer/Public/AudioMixer.h b/Engine/Source/Runtime/AudioMixer/Public/AudioMixer.h
--- a/Engine/Source/Runtime/AudioMixer/Public/AudioMixer.h
+++ b/Engine/Source/Runtime/AudioMixer/Public/AudioMixer.h
@@ -137,7 +137,7 @@ namespace Audio
 		bool bRestoreIfRemoved;
 
 		/* The maximum number of sources we will try to decode or playback at once. */
-		int32 MaxChannels;
+		int32 MaxSources;
 
 		FAudioMixerOpenStreamParams()
 			: OutputDeviceIndex(INDEX_NONE)
@@ -146,7 +146,7 @@ namespace Audio
 			, AudioMixer(nullptr)
 			, SampleRate(44100)
 			, bRestoreIfRemoved(false)
-			, MaxChannels(32)
+			, MaxSources(0)
 		{}
 	};
 
diff --git a/Engine/Source/Runtime/Engine/Private/AudioDevice.cpp b/Engine/Source/Runtime/Engine/Private/AudioDevice.cpp
--- a/Engine/Source/Runtime/Engine/Private/AudioDevice.cpp
+++ b/Engine/Source/Runtime/Engine/Private/AudioDevice.cpp
@@ -199,7 +199,7 @@ FAttenuationListenerData FAttenuationListenerData::Create(const FAudioDevice& Au
 -----------------------------------------------------------------------------*/
 
 FAudioDevice::FAudioDevice()
-	: NumStoppingVoices(32)
+	: NumStoppingSources(32)
 	, SampleRate(0)
 	, NumPrecacheFrames(MONO_PCM_BUFFER_SAMPLES)
 	, CommonAudioPoolSize(0)
@@ -302,11 +302,15 @@ bool FAudioDevice::Init(int32 InMaxSources)
 	// as the Sources array has yet to be initialized. If the cvar is largest, take that value to allow for testing
 	const int32 PlatformMaxSources = PlatformSettings.MaxChannels > 0 ? PlatformSettings.MaxChannels : InMaxSources;
 	MaxSources = FMath::Max(PlatformMaxSources, AudioChannelCountCVar);
-	MaxChannels = MaxSources;
-	MaxChannels_GameThread = MaxSources;
+	MaxSources = FMath::Max(MaxSources, 1);
 
 	// Ensure and not assert so if in editor, user can change quality setting and re-serialize if so desired.
 	ensureMsgf(MaxSources > 0, TEXT("Neither passed MaxSources nor platform MaxChannel setting was positive value"));
+	UE_LOG(LogAudio, Display, TEXT("AudioDevice MaxSources: %d"), MaxSources);
+
+	MaxChannels = MaxSources;
+	MaxChannels_GameThread = MaxSources;
+
 
 	// Mixed sample rate is set by the platform
 	SampleRate = PlatformSettings.SampleRate;
@@ -360,12 +364,12 @@ bool FAudioDevice::Init(int32 InMaxSources)
 		// create a platform specific effects manager
 		Effects = CreateEffectsManager();
 
-		NumStoppingVoices = GetDefault<UAudioSettings>()->NumStoppingSources;
+		NumStoppingSources = GetDefault<UAudioSettings>()->NumStoppingSources;
 	}
 	else
 	{
-		// In old audio engine, there are no stopping voices
-		NumStoppingVoices = 0;
+		// Stopping sources are not supported in the old audio engine
+		NumStoppingSources = 0;
 	}
 
 	// Cache any plugin settings objects we have loaded
@@ -508,29 +512,31 @@ void FAudioDevice::SetMaxChannels(int32 InMaxChannels)
 		return;
 	}
 
-	if (IsInGameThread())
+	if (InMaxChannels > MaxSources)
 	{
-		MaxChannels_GameThread = InMaxChannels;
+		UE_LOG(LogAudio, Warning, TEXT("Can't increase MaxChannels past MaxSources"));
+		return;
 	}
 
-	if (!IsInAudioThread())
+	if (IsInAudioThread())
 	{
-		DECLARE_CYCLE_STAT(TEXT("FAudioThreadTask.SetMaxChannels"), STAT_AudioSetMaxChannels, STATGROUP_AudioThreadCommands);
-
-		FAudioThread::RunCommandOnAudioThread([this, InMaxChannels]()
+		DECLARE_CYCLE_STAT(TEXT("FAudioThreadTask.SetMaxChannelsGameThread"), STAT_AudioSetMaxChannelsGameThread, STATGROUP_AudioThreadCommands);
+		FAudioThread::RunCommandOnGameThread([this, InMaxChannels]()
 		{
-			this->SetMaxChannels(InMaxChannels);
+			MaxChannels_GameThread = InMaxChannels;
 
-		}, GET_STATID(STAT_AudioSetMaxChannels));
+		}, GET_STATID(STAT_AudioSetMaxChannelsGameThread));
 	}
 
-	if (InMaxChannels > MaxSources)
+	if (IsInGameThread())
 	{
-		UE_LOG(LogAudio, Warning, TEXT("Can't increase MaxChannels past MaxSources"));
-		return;
-	}
+		DECLARE_CYCLE_STAT(TEXT("FAudioThreadTask.SetMaxChannels"), STAT_AudioSetMaxChannels, STATGROUP_AudioThreadCommands);
 
-	MaxChannels = InMaxChannels;
+		FAudioThread::RunCommandOnAudioThread([this, InMaxChannels]()
+		{
+			MaxChannels = InMaxChannels;
+		}, GET_STATID(STAT_AudioSetMaxChannels));
+	}
 }
 
 void FAudioDevice::SetMaxChannelsScaled(float InScaledChannelCount)
@@ -577,9 +583,15 @@ int32 FAudioDevice::GetMaxChannels() const
 	}
 
 	// Find product of max channels and final scalar, and clamp between 1 and MaxSources.
+	check(MaxSources > 0);
 	return FMath::Clamp(static_cast<int32>(OutMaxChannels * MaxChannelScalarToApply), 1, MaxSources);
 }
 
+int32 FAudioDevice::GetMaxSources() const
+{
+	return MaxSources + NumStoppingSources;
+}
+
 void FAudioDevice::Teardown()
 {
 	// Do a fadeout to prevent clicking on shutdown
@@ -1958,8 +1970,8 @@ void FAudioDevice::InitSoundSources()
 	if (Sources.Num() == 0)
 	{
 		// now create platform specific sources
-		const int32 Channels = GetMaxChannels() + NumStoppingVoices;
-		for (int32 SourceIndex = 0; SourceIndex < Channels; SourceIndex++)
+		const int32 SourceMax = GetMaxSources();
+		for (int32 SourceIndex = 0; SourceIndex < SourceMax; SourceIndex++)
 		{
 			FSoundSource* Source = CreateSoundSource();
 			Source->InitializeSourceEffects(SourceIndex);
@@ -4012,7 +4024,7 @@ void FAudioDevice::Update(bool bGameTicking)
 		SET_DWORD_STAT(STAT_ActiveSounds, ActiveSounds.Num());
 		SET_DWORD_STAT(STAT_AudioVirtualLoops, VirtualLoops.Num());
 		SET_DWORD_STAT(STAT_AudioMaxChannels, Channels);
-		SET_DWORD_STAT(STAT_AudioMaxStoppingSources, NumStoppingVoices);
+		SET_DWORD_STAT(STAT_AudioMaxStoppingSources, NumStoppingSources);
 	}
 
 	// now let the platform perform anything it needs to handle
diff --git a/Engine/Source/Runtime/Engine/Public/AudioDevice.h b/Engine/Source/Runtime/Engine/Public/AudioDevice.h
--- a/Engine/Source/Runtime/Engine/Public/AudioDevice.h
+++ b/Engine/Source/Runtime/Engine/Public/AudioDevice.h
@@ -553,6 +553,10 @@ public:
 	/** Returns the max channels used by the audio device. */
 	int32 GetMaxChannels() const;
 
+	/** Returns the maximum sources used by the audio device set on initialization,
+	  * including the number of stopping voices reserved. */
+	int32 GetMaxSources() const;
+
 	/**
 	* Stops any sound sources which are using the given buffer.
 	*
@@ -1584,7 +1588,7 @@ public:
 public:
 
 	/** The number of sources to reserve for stopping sounds. */
-	int32 NumStoppingVoices;
+	int32 NumStoppingSources;
 
 	/** The sample rate of all the audio devices */
 	int32 SampleRate;
diff --git a/Engine/Source/Runtime/IOS/IOSAudio/Private/IOSAudioDevice.cpp b/Engine/Source/Runtime/IOS/IOSAudio/Private/IOSAudioDevice.cpp
--- a/Engine/Source/Runtime/IOS/IOSAudio/Private/IOSAudioDevice.cpp
+++ b/Engine/Source/Runtime/IOS/IOSAudio/Private/IOSAudioDevice.cpp
@@ -11,6 +11,7 @@
 #include "IOSAudioDevice.h"
 #include "AudioEffect.h"
 #include "ADPCMAudioInfo.h"
+#include "AudioPluginUtilities.h"
 
 DEFINE_LOG_CATEGORY(LogIOSAudio);
 
@@ -109,7 +110,8 @@ bool FIOSAudioDevice::Exec(UWorld* InWorld, const TCHAR* Cmd, FOutputDevice& Ar)
 
 FAudioPlatformSettings FIOSAudioDevice::GetPlatformSettings() const
 {
-	return FAudioPlatformSettings::GetPlatformSettings(TEXT("/Script/IOSRuntimeSettings.IOSRuntimeSettings"));
+	const TCHAR* ConfigSection = AudioPluginUtilities::GetPlatformConfigSection(EAudioPlatform::IOS);
+	return FAudioPlatformSettings::GetPlatformSettings(ConfigSection);
 }
 
 bool FIOSAudioDevice::InitializeHardware()
@@ -198,7 +200,7 @@ bool FIOSAudioDevice::InitializeHardware()
 		return false;
 	}
 
-	uint32 BusCount = GetMaxChannels() * CHANNELS_PER_BUS;
+	uint32 BusCount = GetMaxSources() * CHANNELS_PER_BUS;
 	Status = AudioUnitSetProperty(MixerUnit,
 	                              kAudioUnitProperty_ElementCount,
 	                              kAudioUnitScope_Input,
diff --git a/Engine/Source/Runtime/Linux/AudioMixerSDL/Private/AudioMixerPlatformSDL.cpp b/Engine/Source/Runtime/Linux/AudioMixerSDL/Private/AudioMixerPlatformSDL.cpp
--- a/Engine/Source/Runtime/Linux/AudioMixerSDL/Private/AudioMixerPlatformSDL.cpp
+++ b/Engine/Source/Runtime/Linux/AudioMixerSDL/Private/AudioMixerPlatformSDL.cpp
@@ -4,6 +4,7 @@
 #include "Modules/ModuleManager.h"
 #include "AudioMixer.h"
 #include "AudioMixerDevice.h"
+#include "AudioPluginUtilities.h"
 #include "CoreGlobals.h"
 #include "Misc/ConfigCacheIni.h"
 #include "OpusAudioInfo.h"
@@ -397,7 +398,8 @@ namespace Audio
 	FAudioPlatformSettings FMixerPlatformSDL::GetPlatformSettings() const
 	{
 #if PLATFORM_UNIX
-		return FAudioPlatformSettings::GetPlatformSettings(TEXT("/Script/LinuxTargetPlatform.LinuxTargetSettings"));
+		const TCHAR* ConfigSection = AudioPluginUtilities::GetPlatformConfigSection(EAudioPlatform::Linux);
+		return FAudioPlatformSettings::GetPlatformSettings(ConfigSection);
 #else
 		// On HTML5 and Windows, use default parameters.
 		return FAudioPlatformSettings();
diff --git a/Engine/Source/Runtime/NonRealtimeAudioRenderer/Private/AudioMixerPlatformNonRealtime.cpp b/Engine/Source/Runtime/NonRealtimeAudioRenderer/Private/AudioMixerPlatformNonRealtime.cpp
--- a/Engine/Source/Runtime/NonRealtimeAudioRenderer/Private/AudioMixerPlatformNonRealtime.cpp
+++ b/Engine/Source/Runtime/NonRealtimeAudioRenderer/Private/AudioMixerPlatformNonRealtime.cpp
@@ -3,6 +3,7 @@
 #include "AudioMixerPlatformNonRealtime.h"
 #include "AudioMixer.h"
 #include "AudioMixerDevice.h"
+#include "AudioPluginUtilities.h"
 #include "HAL/PlatformAffinity.h"
 #include "Misc/App.h"
 
@@ -386,6 +387,7 @@ namespace Audio
 
 	FAudioPlatformSettings FMixerPlatformNonRealtime::GetPlatformSettings() const
 	{
+		const TCHAR* ConfigSection = AudioPluginUtilities::GetPlatformConfigSection(EAudioPlatform::Windows);
 		return FAudioPlatformSettings::GetPlatformSettings(TEXT("/Script/WindowsTargetPlatform.WindowsTargetSettings"));
 	}
 
diff --git a/Engine/Source/Runtime/Windows/AudioMixerXAudio2/Private/AudioMixerPlatformXAudio2.cpp b/Engine/Source/Runtime/Windows/AudioMixerXAudio2/Private/AudioMixerPlatformXAudio2.cpp
--- a/Engine/Source/Runtime/Windows/AudioMixerXAudio2/Private/AudioMixerPlatformXAudio2.cpp
+++ b/Engine/Source/Runtime/Windows/AudioMixerXAudio2/Private/AudioMixerPlatformXAudio2.cpp
@@ -9,6 +9,7 @@
 #include "AudioMixerPlatformXAudio2.h"
 #include "AudioMixer.h"
 #include "AudioMixerDevice.h"
+#include "AudioPluginUtilities.h"
 #include "HAL/PlatformAffinity.h"
 
 #ifndef WITH_XMA2
@@ -985,6 +986,7 @@ namespace Audio
 
 	FAudioPlatformSettings FMixerPlatformXAudio2::GetPlatformSettings() const
 	{
+		const TCHAR* ConfigSection = AudioPluginUtilities::GetPlatformConfigSection(EAudioPlatform::Windows);
 		return FAudioPlatformSettings::GetPlatformSettings(TEXT("/Script/WindowsTargetPlatform.WindowsTargetSettings"));
 	}
 
diff --git a/Engine/Source/Runtime/Windows/XAudio2/Private/XAudio2Device.cpp b/Engine/Source/Runtime/Windows/XAudio2/Private/XAudio2Device.cpp
--- a/Engine/Source/Runtime/Windows/XAudio2/Private/XAudio2Device.cpp
+++ b/Engine/Source/Runtime/Windows/XAudio2/Private/XAudio2Device.cpp
@@ -13,6 +13,7 @@
 
 #include "XAudio2Device.h"
 #include "AudioEffect.h"
+#include "AudioPluginUtilities.h"
 #include "OpusAudioInfo.h"
 #include "VorbisAudioInfo.h"
 #include "ADPCMAudioInfo.h"
@@ -744,5 +745,6 @@ void* FXAudio2Device::AllocatePermanentMemory( int32 Size, bool& AllocatedInPool
 
 FAudioPlatformSettings FXAudio2Device::GetPlatformSettings() const
 {
-	return FAudioPlatformSettings::GetPlatformSettings(TEXT("/Script/WindowsTargetPlatform.WindowsTargetSettings"));
+	const TCHAR* ConfigSection = AudioPluginUtilities::GetPlatformConfigSection(EAudioPlatform::Windows);
+	return FAudioPlatformSettings::GetPlatformSettings(ConfigSection);
 }
