[NO NON-PII COMMIT MESSAGE]

diff --git a/Engine/Source/Runtime/Android/AudioMixerAndroid/Private/AudioMixerAndroid.cpp b/Engine/Source/Runtime/Android/AudioMixerAndroid/Private/AudioMixerAndroid.cpp
--- a/Engine/Source/Runtime/Android/AudioMixerAndroid/Private/AudioMixerAndroid.cpp
+++ b/Engine/Source/Runtime/Android/AudioMixerAndroid/Private/AudioMixerAndroid.cpp
@@ -7,6 +7,8 @@
 class FAudioMixerModuleAndroid : public IAudioDeviceModule
 {
 public:
+	virtual bool IsAudioMixerModule() const override { return true; }
+
 	virtual Audio::IAudioMixerPlatformInterface* CreateAudioMixerPlatformInterface() override
 	{
 		return new Audio::FMixerPlatformAndroid();
