[NO NON-PII COMMIT MESSAGE]

diff --git a/Engine/Source/Runtime/Android/AudioMixerAndroid/Private/AudioMixerPlatformAndroid.cpp b/Engine/Source/Runtime/Android/AudioMixerAndroid/Private/AudioMixerPlatformAndroid.cpp
--- a/Engine/Source/Runtime/Android/AudioMixerAndroid/Private/AudioMixerPlatformAndroid.cpp
+++ b/Engine/Source/Runtime/Android/AudioMixerAndroid/Private/AudioMixerPlatformAndroid.cpp
@@ -233,6 +233,8 @@ namespace Audio
 			return false;
 		}
 
+		AudioStreamInfo.DeviceInfo.SampleRate = OpenStreamParams.SampleRate;
+
 		SLresult Result;
 
 		FAudioPlatformSettings PlatformSettings = GetPlatformSettings();
