[Partner-Oculus] Select Quest by default when adding new Oculus target device

[CL 9384857 by kevin.xiao in Partner-Oculus-4.23-Private]
[CL 9461535 by [REDACTED_PERSON] in Partner-Oculus-Private]

diff --git a/Engine/Source/Runtime/Android/AndroidRuntimeSettings/Private/AndroidRuntimeSettings.cpp b/Engine/Source/Runtime/Android/AndroidRuntimeSettings/Private/AndroidRuntimeSettings.cpp
--- a/Engine/Source/Runtime/Android/AndroidRuntimeSettings/Private/AndroidRuntimeSettings.cpp
+++ b/Engine/Source/Runtime/Android/AndroidRuntimeSettings/Private/AndroidRuntimeSettings.cpp
@@ -146,6 +146,27 @@ void UAndroidRuntimeSettings::PostEditChangeProperty(struct FPropertyChangedEven
 		}
 	}
 
+	if (PropertyChangedEvent.Property != nullptr && PropertyChangedEvent.Property->GetName().StartsWith(TEXT("PackageForOculusMobile")))
+	{
+		if (PropertyChangedEvent.ChangeType == EPropertyChangeType::ArrayAdd)
+		{
+			// Get a list of all available devices
+			TArray<EOculusMobileDevice::Type> deviceList;
+#define OCULUS_DEVICE_LOOP(device) deviceList.Add(device);
+			FOREACH_ENUM_EOCULUSMOBILEDEVICE(OCULUS_DEVICE_LOOP);
+#undef OCULUS_DEVICE_LOOP
+			// Add last device that isn't already in the list
+			for (int i = deviceList.Num() - 1; i >= 0; --i)
+			{
+				if (!PackageForOculusMobile.Contains(deviceList[i]))
+				{
+					PackageForOculusMobile.Last() = deviceList[i];
+					break;
+				}
+			}
+		}
+	}
+
 	HandlesRGBHWSupport();
 }
 
