Update Android settings to indicate 19 is minimum Android version

#ue4
#android

diff --git a/Engine/Source/Runtime/Android/AndroidRuntimeSettings/Classes/AndroidRuntimeSettings.h b/Engine/Source/Runtime/Android/AndroidRuntimeSettings/Classes/AndroidRuntimeSettings.h
--- a/Engine/Source/Runtime/Android/AndroidRuntimeSettings/Classes/AndroidRuntimeSettings.h
+++ b/Engine/Source/Runtime/Android/AndroidRuntimeSettings/Classes/AndroidRuntimeSettings.h
@@ -205,12 +205,12 @@ public:
 	UPROPERTY(GlobalConfig, EditAnywhere, Category = "APK Packaging", Meta = (DisplayName = "Version Display Name (usually x.y)"))
 	FString VersionDisplayName;
 
-	// What OS version the app is allowed to be installed on (do not set this lower than 9)
-	UPROPERTY(GlobalConfig, EditAnywhere, Category = "APK Packaging", Meta = (DisplayName = "Minimum SDK Version (9=Gingerbread, 14=Ice Cream Sandwich, 21=Lollipop)"))
+	// What OS version the app is allowed to be installed on (do not set this lower than 19)
+	UPROPERTY(GlobalConfig, EditAnywhere, Category = "APK Packaging", Meta = (DisplayName = "Minimum SDK Version (19=KitKat, 21=Lollipop)"))
 	int32 MinSDKVersion;
 	
-	// What OS version the app is expected to run on (do not set this lower than 9, set to 19 for Gear VR)
-	UPROPERTY(GlobalConfig, EditAnywhere, Category = "APK Packaging", Meta = (DisplayName = "Target SDK Version (9=Gingerbread, 14=Ice Cream Sandwich, 21=Lollipop)"))
+	// What OS version the app is expected to run on (do not set this lower than 19)
+	UPROPERTY(GlobalConfig, EditAnywhere, Category = "APK Packaging", Meta = (DisplayName = "Target SDK Version (19=KitKat, 21=Lollipop)"))
 	int32 TargetSDKVersion;
 
 	// Preferred install location for the application
