Fix AndroidRuntime string referencing legacy UE4 Java class

diff --git a/Engine/Source/Runtime/Android/AndroidRuntimeSettings/Classes/AndroidRuntimeSettings.h b/Engine/Source/Runtime/Android/AndroidRuntimeSettings/Classes/AndroidRuntimeSettings.h
--- a/Engine/Source/Runtime/Android/AndroidRuntimeSettings/Classes/AndroidRuntimeSettings.h
+++ b/Engine/Source/Runtime/Android/AndroidRuntimeSettings/Classes/AndroidRuntimeSettings.h
@@ -344,7 +344,7 @@ public:
 	UPROPERTY(GlobalConfig, EditAnywhere, Category = "Advanced APK Packaging", Meta = (DisplayName = "Extra Settings for <application> section (\\n to separate lines)"))
 	FString ExtraApplicationSettings;
 
-	UPROPERTY(GlobalConfig, EditAnywhere, Category = "Advanced APK Packaging", Meta = (DisplayName = "Extra Tags for UE4.GameActivity <activity> node"))
+	UPROPERTY(GlobalConfig, EditAnywhere, Category = "Advanced APK Packaging", Meta = (DisplayName = "Extra Tags for com.epicgames.unreal.GameActivity <activity> node"))
 	TArray<FString> ExtraActivityNodeTags;
 
 	// Any extra settings for the main <activity> section (an optional file <Project>/Build/Android/ManifestApplicationActivtyAdditions.txt will also be included)
