Change display of OpenGL 3.1 support to 3.2 for Android project settings

#UE-140707
#android

diff --git a/Engine/Source/Runtime/Android/AndroidRuntimeSettings/Classes/AndroidRuntimeSettings.h b/Engine/Source/Runtime/Android/AndroidRuntimeSettings/Classes/AndroidRuntimeSettings.h
--- a/Engine/Source/Runtime/Android/AndroidRuntimeSettings/Classes/AndroidRuntimeSettings.h
+++ b/Engine/Source/Runtime/Android/AndroidRuntimeSettings/Classes/AndroidRuntimeSettings.h
@@ -363,8 +363,8 @@ public:
 	UPROPERTY(GlobalConfig, EditAnywhere, Category = Build, meta = (DisplayName = "Support x86_64 [aka x64]"))
 	bool bBuildForX8664;
 
-	// Include shaders for devices supporting OpenGL ES 3.1 and above (default)
-	UPROPERTY(GlobalConfig, EditAnywhere, Category = Build, meta = (DisplayName = "Support OpenGL ES3.1"))
+	// Include shaders for devices supporting OpenGL ES 3.2 and above (default)
+	UPROPERTY(GlobalConfig, EditAnywhere, Category = Build, meta = (DisplayName = "Support OpenGL ES3.2"))
 	bool bBuildForES31;
 
 	// Support the Vulkan RHI and include Vulkan shaders
