Removing AudioStreamingThreshold platform setting (has no effect).

diff --git a/Engine/Source/Developer/Windows/WindowsTargetPlatform/Classes/WindowsTargetSettings.h b/Engine/Source/Developer/Windows/WindowsTargetPlatform/Classes/WindowsTargetSettings.h
--- a/Engine/Source/Developer/Windows/WindowsTargetPlatform/Classes/WindowsTargetSettings.h
+++ b/Engine/Source/Developer/Windows/WindowsTargetPlatform/Classes/WindowsTargetSettings.h
@@ -136,7 +136,7 @@ public:
 	float CompressionQualityModifier;
 
 	/** When set to anything beyond 0, this will ensure any SoundWaves longer than this value, in seconds, to stream directly off of the disk. */
-	UPROPERTY(GlobalConfig, EditAnywhere, Category = "Audio|CookOverrides", meta = (DisplayName = "Stream All Soundwaves Longer Than: "))
+	UPROPERTY(GlobalConfig)
 	float AutoStreamingThreshold;
 
 	/** Quality Level to COOK SoundCues at (if set, all other levels will be stripped by the cooker). */
diff --git a/Engine/Source/Runtime/Android/AndroidRuntimeSettings/Classes/AndroidRuntimeSettings.h b/Engine/Source/Runtime/Android/AndroidRuntimeSettings/Classes/AndroidRuntimeSettings.h
--- a/Engine/Source/Runtime/Android/AndroidRuntimeSettings/Classes/AndroidRuntimeSettings.h
+++ b/Engine/Source/Runtime/Android/AndroidRuntimeSettings/Classes/AndroidRuntimeSettings.h
@@ -555,7 +555,7 @@ public:
 	float CompressionQualityModifier;
 
 	// When set to anything beyond 0, this will ensure any SoundWaves longer than this value, in seconds, to stream directly off of the disk.
-	UPROPERTY(GlobalConfig, EditAnywhere, Category = "Audio|CookOverrides", meta = (DisplayName = "Stream All Soundwaves Longer Than: "))
+	UPROPERTY(GlobalConfig)
 	float AutoStreamingThreshold;
 
 	// Several Android graphics debuggers require configuration changes to be made to your application in order to operate. Choosing an option from this menu will configure your project to work with that graphics debugger. 
diff --git a/Engine/Source/Runtime/IOS/IOSRuntimeSettings/Classes/IOSRuntimeSettings.h b/Engine/Source/Runtime/IOS/IOSRuntimeSettings/Classes/IOSRuntimeSettings.h
--- a/Engine/Source/Runtime/IOS/IOSRuntimeSettings/Classes/IOSRuntimeSettings.h
+++ b/Engine/Source/Runtime/IOS/IOSRuntimeSettings/Classes/IOSRuntimeSettings.h
@@ -551,7 +551,7 @@ public:
 	float CompressionQualityModifier;
 
 	// When set to anything beyond 0, this will ensure any SoundWaves longer than this value, in seconds, to stream directly off of the disk.
-	UPROPERTY(GlobalConfig, EditAnywhere, Category = "Audio|CookOverrides", meta = (DisplayName = "Stream All Soundwaves Longer Than: "))
+	UPROPERTY(GlobalConfig)
 	float AutoStreamingThreshold;
 
 	/** Whether to enable LOD streaming for landscape visual meshes. Requires Metal support. */
