Add Quest 3 to list of supported devices.

diff --git a/Engine/Source/Runtime/Android/AndroidRuntimeSettings/Private/AndroidRuntimeSettings.cpp b/Engine/Source/Runtime/Android/AndroidRuntimeSettings/Private/AndroidRuntimeSettings.cpp
--- a/Engine/Source/Runtime/Android/AndroidRuntimeSettings/Private/AndroidRuntimeSettings.cpp
+++ b/Engine/Source/Runtime/Android/AndroidRuntimeSettings/Private/AndroidRuntimeSettings.cpp
@@ -134,7 +134,7 @@ void UAndroidRuntimeSettings::HandleMetaQuestSupport()
 
 		if (ExtraApplicationSettings.Find("com.oculus.supportedDevices") == INDEX_NONE)
 		{
-			FString SupportedDevicesValue("quest|quest2|questpro");
+			FString SupportedDevicesValue("quest|quest2|questpro|quest3");
 			ExtraApplicationSettings.Append("<meta-data android:name=\"com.oculus.supportedDevices\" android:value=\"" + SupportedDevicesValue + "\" />");
 			UpdateSinglePropertyInConfigFile(GetClass()->FindPropertyByName(GET_MEMBER_NAME_CHECKED(UAndroidRuntimeSettings, ExtraApplicationSettings)), GetDefaultConfigFilename());
 		}
