Fix quest 3 not showing up in supported devices list.

diff --git a/Engine/Source/Runtime/Android/AndroidRuntimeSettings/Private/AndroidRuntimeSettings.cpp b/Engine/Source/Runtime/Android/AndroidRuntimeSettings/Private/AndroidRuntimeSettings.cpp
--- a/Engine/Source/Runtime/Android/AndroidRuntimeSettings/Private/AndroidRuntimeSettings.cpp
+++ b/Engine/Source/Runtime/Android/AndroidRuntimeSettings/Private/AndroidRuntimeSettings.cpp
@@ -85,11 +85,13 @@ void UAndroidRuntimeSettings::HandleMetaQuestSupport()
 	// PackageForOculusMobile doesn't get loaded since it's marked as deprecated, so it needs to be read directly from the config
 	TArray<FString> PackageList;
 	GConfig->GetArray(TEXT("/Script/AndroidRuntimeSettings.AndroidRuntimeSettings"), *FString("+").Append(GET_MEMBER_NAME_STRING_CHECKED(UAndroidRuntimeSettings, PackageForOculusMobile)), PackageList, GetDefaultConfigFilename());
+
+	FString SupportedDevicesTag("<meta-data android:name=\"com.oculus.supportedDevices\"");
 	if (PackageList.Num() > 0)
 	{
 		bPackageForMetaQuest = true;
 		// Clean ExtraApplications metadata so that the updated list of supported devices will be added further down.
-		RemoveExtraApplicationTag("<meta-data android:name=\"com.oculus.supportedDevices\"");
+		RemoveExtraApplicationTag(SupportedDevicesTag);
 		// Use TryUpdateDefaultConfigFile() instead of UpdateSinglePropertyInConfigFile() so that the PackageForOculusMobile will also get cleared
 		TryUpdateDefaultConfigFile();
 	}
@@ -132,9 +134,19 @@ void UAndroidRuntimeSettings::HandleMetaQuestSupport()
 		UE_LOG(LogAndroidRuntimeSettings, Display, TEXT("Support Vulkan Desktop: %d.\n"), bSupportsVulkanSM5);
 		UE_LOG(LogAndroidRuntimeSettings, Display, TEXT("Support OpenGL ES3.2: %d."), bBuildForES31);
 
-		if (ExtraApplicationSettings.Find("com.oculus.supportedDevices") == INDEX_NONE)
+		int32 SupportedDevicesTagIndex = ExtraApplicationSettings.Find("com.oculus.supportedDevices");
+		FString SupportedDevicesValue("quest|quest2|questpro|quest3");
+		int32 SupportedDevicesIndex = ExtraApplicationSettings.Find(SupportedDevicesValue);
+		// The supported devices tag is present but not up to date and does not contain all the currently supported devices.
+		bool bNeedtoUpdateDevices = (SupportedDevicesTagIndex != INDEX_NONE) && (SupportedDevicesIndex == INDEX_NONE);
+		// Remove the current supported devices value so that it can be added again with the updated value further down.
+		if (bNeedtoUpdateDevices)
+		{
+			RemoveExtraApplicationTag(SupportedDevicesTag);
+		}
+
+		if (SupportedDevicesTagIndex == INDEX_NONE || bNeedtoUpdateDevices)
 		{
-			FString SupportedDevicesValue("quest|quest2|questpro|quest3");
 			ExtraApplicationSettings.Append("<meta-data android:name=\"com.oculus.supportedDevices\" android:value=\"" + SupportedDevicesValue + "\" />");
 			UpdateSinglePropertyInConfigFile(GetClass()->FindPropertyByName(GET_MEMBER_NAME_CHECKED(UAndroidRuntimeSettings, ExtraApplicationSettings)), GetDefaultConfigFilename());
 		}
@@ -142,7 +154,7 @@ void UAndroidRuntimeSettings::HandleMetaQuestSupport()
 	else 
 	{
 		// Clean up the supported devices metadata tag so it doesn't end up in the manifest when bPackageForMetaQuest is turned off.
-		RemoveExtraApplicationTag("<meta-data android:name=\"com.oculus.supportedDevices\"");
+		RemoveExtraApplicationTag(SupportedDevicesTag);
 	}
 }
 
