Removing defunct Format setting on Android

#tests tested editor with android sdk installed to verify
[RN] Remove defunt Audio Encoding setting from Android

diff --git a/Engine/Source/Runtime/Android/AndroidRuntimeSettings/Classes/AndroidRuntimeSettings.h b/Engine/Source/Runtime/Android/AndroidRuntimeSettings/Classes/AndroidRuntimeSettings.h
--- a/Engine/Source/Runtime/Android/AndroidRuntimeSettings/Classes/AndroidRuntimeSettings.h
+++ b/Engine/Source/Runtime/Android/AndroidRuntimeSettings/Classes/AndroidRuntimeSettings.h
@@ -165,7 +165,7 @@ struct FGooglePlayLeaderboardMapping
 UENUM()
 namespace EAndroidAudio
 {
-	enum Type : int
+	enum UE_DEPRECATED(5.5, "Please don't use") Type : int
 	{
 		Default = 0 UMETA(DisplayName = "Default", ToolTip = "This option selects the default encoder."),
 		OGG = 1 UMETA(DisplayName = "Ogg Vorbis", ToolTip = "Selects Ogg Vorbis encoding."),
@@ -574,10 +574,6 @@ public:
 	UPROPERTY(GlobalConfig, EditAnywhere, Category = Input, meta = (DisplayName = "Block force feedback on the device when controllers are attached"))
 	bool bControllersBlockDeviceFeedback;
 
-	/** Android encoding options. */
-	UPROPERTY(GlobalConfig, EditAnywhere, Category = Audio, meta = (DisplayName = "Encoding Format"))
-	TEnumAsByte<EAndroidAudio::Type> AndroidAudio;
-
 	/** Sample rate to run the audio mixer with. */
 	UPROPERTY(config, EditAnywhere, Category = "Audio", Meta = (DisplayName = "Audio Mixer Sample Rate"))
 	int32 AudioSampleRate;
