Removed "Quest" from ApplicationSettings when Package for Meta Quest is checked. Quest 1 is no longer supported.

diff --git a/Engine/Source/Runtime/Android/AndroidRuntimeSettings/Private/AndroidRuntimeSettings.cpp b/Engine/Source/Runtime/Android/AndroidRuntimeSettings/Private/AndroidRuntimeSettings.cpp
--- a/Engine/Source/Runtime/Android/AndroidRuntimeSettings/Private/AndroidRuntimeSettings.cpp
+++ b/Engine/Source/Runtime/Android/AndroidRuntimeSettings/Private/AndroidRuntimeSettings.cpp
@@ -137,7 +137,7 @@ void UAndroidRuntimeSettings::HandleMetaQuestSupport()
 		UE_LOG(LogAndroidRuntimeSettings, Display, TEXT("Support OpenGL ES3.2: %d."), bBuildForES31);
 
 		int32 SupportedDevicesTagIndex = ExtraApplicationSettings.Find("com.oculus.supportedDevices");
-		FString SupportedDevicesValue("quest|quest2|questpro|quest3");
+		FString SupportedDevicesValue("quest2|questpro|quest3");
 		int32 SupportedDevicesIndex = ExtraApplicationSettings.Find(SupportedDevicesValue);
 		// The supported devices tag is present but not up to date and does not contain all the currently supported devices.
 		bool bNeedtoUpdateDevices = (SupportedDevicesTagIndex != INDEX_NONE) && (SupportedDevicesIndex == INDEX_NONE);
